Move WFC Plugin to /Plugins/Experimental/
Add/Delete instead of Rename/Move

diff --git a/Engine/Plugins/Experimental/WaveFunctionCollapse/Source/WaveFunctionCollapse/Private/WaveFunctionCollapse.cpp b/Engine/Plugins/Experimental/WaveFunctionCollapse/Source/WaveFunctionCollapse/Private/WaveFunctionCollapse.cpp
new file mode 100644
--- /dev/null
+++ b/Engine/Plugins/Experimental/WaveFunctionCollapse/Source/WaveFunctionCollapse/Private/WaveFunctionCollapse.cpp
@@ -0,0 +1,63 @@
+// Copyright Epic Games, Inc. All Rights Reserved.
+
+#include "WaveFunctionCollapse.h"
+#include "WaveFunctionCollapseEditorCommands.h"
+#include "LevelEditor.h"
+#include "EditorUtilityWidget.h"
+#include "WidgetBlueprint.h"
+#include "EditorUtilitySubsystem.h"
+#include "EditorUtilityWidgetBlueprint.h"
+#include "UObject/UObjectGlobals.h"
+#include "Framework/MultiBox/MultiBoxBuilder.h"
+
+#define LOCTEXT_NAMESPACE "FWaveFunctionCollapseModule"
+
+void FWaveFunctionCollapseModule::StartupModule()
+{
+	// This code will execute after your module is loaded into memory; the exact timing is specified in the .uplugin file per-module
+	FWaveFunctionCollapseEditorCommands::Register();
+
+	PluginCommands = MakeShareable(new FUICommandList);
+
+	PluginCommands->MapAction(
+		FWaveFunctionCollapseEditorCommands::Get().WaveFunctionCollapseWidget,
+		FExecuteAction::CreateRaw(this, &FWaveFunctionCollapseModule::WaveFunctionCollapseUI)
+	);
+
+	FLevelEditorModule& LevelEditorModule = FModuleManager::LoadModuleChecked<FLevelEditorModule>("LevelEditor");
+
+	{
+		TSharedPtr<FExtender> NewMenuExtender = MakeShareable(new FExtender);
+		NewMenuExtender->AddMenuExtension("LevelEditor",
+			EExtensionHook::After,
+			PluginCommands,
+			FMenuExtensionDelegate::CreateRaw(this, &FWaveFunctionCollapseModule::AddMenuEntry));
+
+		LevelEditorModule.GetMenuExtensibilityManager()->AddExtender(NewMenuExtender);
+	}
+}
+
+void FWaveFunctionCollapseModule::ShutdownModule()
+{
+	// This function may be called during shutdown to clean up your module.  For modules that support dynamic reloading,
+	// we call this function before unloading the module.
+	FWaveFunctionCollapseEditorCommands::Unregister();
+}
+
+void FWaveFunctionCollapseModule::AddMenuEntry(FMenuBuilder& MenuBuilder)
+{
+	MenuBuilder.BeginSection("CustomMenu", TAttribute<FText>(FText::FromString("WaveFunctionCollapse")));
+	MenuBuilder.AddMenuEntry(FWaveFunctionCollapseEditorCommands::Get().WaveFunctionCollapseWidget);
+	MenuBuilder.EndSection();
+}
+
+void FWaveFunctionCollapseModule::WaveFunctionCollapseUI()
+{
+	UEditorUtilitySubsystem* EUSubsystem = GEditor->GetEditorSubsystem<UEditorUtilitySubsystem>();
+	UEditorUtilityWidgetBlueprint* WaveFunctionCollapseEUWBP = LoadObject<UEditorUtilityWidgetBlueprint>(NULL, TEXT("/WaveFunctionCollapse/Core/WaveFunctionCollapse.WaveFunctionCollapse"), NULL, LOAD_None, NULL);
+	UEditorUtilityWidget* WaveFunctionCollapseEUW = EUSubsystem->SpawnAndRegisterTab(WaveFunctionCollapseEUWBP);
+}
+
+#undef LOCTEXT_NAMESPACE
+	
+IMPLEMENT_MODULE(FWaveFunctionCollapseModule, WaveFunctionCollapse)
\ No newline at end of file
diff --git a/Engine/Plugins/Experimental/WaveFunctionCollapse/Source/WaveFunctionCollapse/Private/WaveFunctionCollapseBPLibrary.cpp b/Engine/Plugins/Experimental/WaveFunctionCollapse/Source/WaveFunctionCollapse/Private/WaveFunctionCollapseBPLibrary.cpp
new file mode 100644
--- /dev/null
+++ b/Engine/Plugins/Experimental/WaveFunctionCollapse/Source/WaveFunctionCollapse/Private/WaveFunctionCollapseBPLibrary.cpp
@@ -0,0 +1,602 @@
+// Copyright Epic Games, Inc. All Rights Reserved.
+
+#include "WaveFunctionCollapseBPLibrary.h"
+#include "WaveFunctionCollapseSubsystem.h"
+#include "Math/UnrealMathUtility.h"
+#include "Engine/StaticMeshActor.h"
+#include "Components/InstancedStaticMeshComponent.h"
+
+UWaveFunctionCollapseBPLibrary::UWaveFunctionCollapseBPLibrary(const FObjectInitializer& ObjectInitializer)
+: Super(ObjectInitializer)
+{
+}
+
+float UWaveFunctionCollapseBPLibrary::CalculateShannonEntropy(const TArray<FWaveFunctionCollapseOption>& Options, UWaveFunctionCollapseModel* WFCModel)
+{
+	if (Options.IsEmpty())
+	{
+		UE_LOG(LogWFC, Display, TEXT("Cannot calculate shannon entropy because the options are empty."));
+		return -1;
+	}
+
+	float SumWeights = 0;
+	float SumWeightXLogWeight = 0;
+	for (const FWaveFunctionCollapseOption& Option : Options)
+	{
+		if (WFCModel->Constraints.Contains(Option))
+		{
+			const float& Weight = WFCModel->Constraints.FindRef(Option).Weight;
+			SumWeights += Weight;
+			SumWeightXLogWeight += (Weight * log(Weight));
+		}
+	}
+
+	if (SumWeights == 0)
+	{
+		UE_LOG(LogWFC, Display, TEXT("Cannot calculate shannon entropy because the sum of weights equals zero."));
+		return -1;
+	}
+
+	return log(SumWeights) - (SumWeightXLogWeight / SumWeights);
+}
+
+int32 UWaveFunctionCollapseBPLibrary::PositionAsIndex(FIntVector Position, FIntVector Resolution)
+{
+	return Position.X + (Position.Y * Resolution.X) + (Position.Z * Resolution.X * Resolution.Y);
+}
+
+FIntVector UWaveFunctionCollapseBPLibrary::IndexAsPosition(int32 Index, FIntVector Resolution)
+{
+	FIntVector Position;
+	Position.Z = floor(Index / (Resolution.X * Resolution.Y));
+	Position.Y = floor((Index - (Position.Z * Resolution.X * Resolution.Y)) / Resolution.X);
+	Position.X = Index - (Position.Y * Resolution.X) - (Position.Z * Resolution.X * Resolution.Y);
+	return Position;
+}
+
+FWaveFunctionCollapseTile UWaveFunctionCollapseBPLibrary::BuildInitialTile(UWaveFunctionCollapseModel* WFCModel)
+{
+	FWaveFunctionCollapseTile InitialTile;
+	for (TPair<FWaveFunctionCollapseOption, FWaveFunctionCollapseAdjacencyToOptionsMap>& Constraint : WFCModel->Constraints)
+	{
+		InitialTile.RemainingOptions.AddUnique(Constraint.Key);
+	}
+	InitialTile.ShannonEntropy = UWaveFunctionCollapseBPLibrary::CalculateShannonEntropy(InitialTile.RemainingOptions, WFCModel);
+	return InitialTile;
+}
+
+TMap<int32, EWaveFunctionCollapseAdjacency> UWaveFunctionCollapseBPLibrary::GetAdjacentIndices(int32 Index, FIntVector Resolution)
+{
+	TMap<int32, EWaveFunctionCollapseAdjacency> AdjacentIndices;
+	FIntVector Position = IndexAsPosition(Index, Resolution);
+	if (Position.X + 1 < Resolution.X) 
+	{ 
+		AdjacentIndices.Add(PositionAsIndex(Position + FIntVector(1, 0, 0), Resolution),EWaveFunctionCollapseAdjacency::Front);
+	}
+	if (Position.X - 1 >= 0) 
+	{ 
+		AdjacentIndices.Add(PositionAsIndex(Position + FIntVector(-1, 0, 0), Resolution), EWaveFunctionCollapseAdjacency::Back); 
+	}
+	if (Position.Y + 1 < Resolution.Y) 
+	{ 
+		AdjacentIndices.Add(PositionAsIndex(Position + FIntVector(0, 1, 0), Resolution), EWaveFunctionCollapseAdjacency::Right); 
+	}
+	if (Position.Y - 1 >= 0) 
+	{ 
+		AdjacentIndices.Add(PositionAsIndex(Position + FIntVector(0, -1, 0), Resolution), EWaveFunctionCollapseAdjacency::Left); 
+	}
+	if (Position.Z + 1 < Resolution.Z) 
+	{ 
+		AdjacentIndices.Add(PositionAsIndex(Position + FIntVector(0, 0, 1), Resolution), EWaveFunctionCollapseAdjacency::Up); 
+	}
+	if (Position.Z - 1 >= 0) 
+	{ 
+		AdjacentIndices.Add(PositionAsIndex(Position + FIntVector(0, 0, -1), Resolution), EWaveFunctionCollapseAdjacency::Down); 
+	}
+
+	return AdjacentIndices;
+}
+
+TMap<FIntVector, EWaveFunctionCollapseAdjacency> UWaveFunctionCollapseBPLibrary::GetAdjacentPositions(FIntVector Position, FIntVector Resolution)
+{
+	TMap<FIntVector, EWaveFunctionCollapseAdjacency> AdjacentPositions;
+	if (Position.X + 1 < Resolution.X) 
+	{ 
+		AdjacentPositions.Add(Position + FIntVector(1, 0, 0), EWaveFunctionCollapseAdjacency::Front); 
+	}
+	if (Position.X - 1 >= 0) 
+	{ 
+		AdjacentPositions.Add(Position + FIntVector(-1, 0, 0), EWaveFunctionCollapseAdjacency::Back); 
+	}
+	if (Position.Y + 1 < Resolution.Y) 
+	{ 
+		AdjacentPositions.Add(Position + FIntVector(0, 1, 0), EWaveFunctionCollapseAdjacency::Right); 
+	}
+	if (Position.Y - 1 >= 0) 
+	{ 
+		AdjacentPositions.Add(Position + FIntVector(0, -1, 0), EWaveFunctionCollapseAdjacency::Left); 
+	}
+	if (Position.Z + 1 < Resolution.Z) 
+	{ 
+		AdjacentPositions.Add(Position + FIntVector(0, 0, 1), EWaveFunctionCollapseAdjacency::Up); 
+	}
+	if (Position.Z - 1 >= 0) 
+	{ 
+		AdjacentPositions.Add(Position + FIntVector(0, 0, -1), EWaveFunctionCollapseAdjacency::Down); 
+	}
+
+	return AdjacentPositions;
+}
+
+bool UWaveFunctionCollapseBPLibrary::IsOptionContained(const FWaveFunctionCollapseOption& Option, const TArray<FWaveFunctionCollapseOption>& Options)
+{
+	for (const FWaveFunctionCollapseOption& CheckAgainst : Options)
+	{
+		if (Option == CheckAgainst)
+		{
+			return true;
+		}
+	}
+	return false;
+}
+
+EWaveFunctionCollapseAdjacency UWaveFunctionCollapseBPLibrary::GetOppositeAdjacency(EWaveFunctionCollapseAdjacency Adjacency)
+{
+	switch (Adjacency)
+	{
+		case EWaveFunctionCollapseAdjacency::Front:
+			return EWaveFunctionCollapseAdjacency::Back;
+		case EWaveFunctionCollapseAdjacency::Back:
+			return EWaveFunctionCollapseAdjacency::Front;
+		case EWaveFunctionCollapseAdjacency::Right:
+			return EWaveFunctionCollapseAdjacency::Left;
+		case EWaveFunctionCollapseAdjacency::Left:
+			return EWaveFunctionCollapseAdjacency::Right;
+		case EWaveFunctionCollapseAdjacency::Up:
+			return EWaveFunctionCollapseAdjacency::Down;
+		case EWaveFunctionCollapseAdjacency::Down:
+			return EWaveFunctionCollapseAdjacency::Up;
+		default:
+			return Adjacency;
+	}
+}
+
+EWaveFunctionCollapseAdjacency UWaveFunctionCollapseBPLibrary::GetNextZAxisClockwiseAdjacency(EWaveFunctionCollapseAdjacency Adjacency)
+{
+	switch (Adjacency)
+	{
+	case EWaveFunctionCollapseAdjacency::Front:
+		return EWaveFunctionCollapseAdjacency::Right;
+	case EWaveFunctionCollapseAdjacency::Right:
+		return EWaveFunctionCollapseAdjacency::Back;
+	case EWaveFunctionCollapseAdjacency::Back:
+		return EWaveFunctionCollapseAdjacency::Left;
+	case EWaveFunctionCollapseAdjacency::Left:
+		return EWaveFunctionCollapseAdjacency::Front;
+	case EWaveFunctionCollapseAdjacency::Up:
+		return EWaveFunctionCollapseAdjacency::Up;
+	case EWaveFunctionCollapseAdjacency::Down:
+		return EWaveFunctionCollapseAdjacency::Down;
+	default:
+		return Adjacency;
+	}
+}
+
+void UWaveFunctionCollapseBPLibrary::AddToAdjacencyToOptionsMap(UPARAM(ref) FWaveFunctionCollapseAdjacencyToOptionsMap& AdjacencyToOptionsMap, EWaveFunctionCollapseAdjacency OptionAdjacency, FWaveFunctionCollapseOption OptionObject)
+{
+	if (FWaveFunctionCollapseOptions* Options = AdjacencyToOptionsMap.AdjacencyToOptionsMap.Find(OptionAdjacency))
+	{
+		// If the Options array does not contain the Option, create it
+		if (!Options->Options.Contains(OptionObject))
+		{
+			Options->Options.Add(OptionObject);
+		}
+	}
+	else // Create the Option and add it to the AdjacencyToOptionsMap
+	{
+		FWaveFunctionCollapseOptions NewOptions;
+		NewOptions.Options.Add(OptionObject);
+		AdjacencyToOptionsMap.AdjacencyToOptionsMap.Add(OptionAdjacency, NewOptions);
+	}
+}
+
+bool UWaveFunctionCollapseBPLibrary::IsSoftObjPathEqual(const FSoftObjectPath& SoftObjectPathA, const FSoftObjectPath& SoftObjectPathB)
+{
+	return SoftObjectPathA == SoftObjectPathB;
+}
+
+FRotator UWaveFunctionCollapseBPLibrary::SanitizeRotator(FRotator Rotator)
+{
+	FRotator OutputRotator;
+	OutputRotator = Rotator;
+
+	// Round to orthogonal values
+	OutputRotator.Roll = FMath::RoundHalfToEven(OutputRotator.Roll / 90) * 90;
+	OutputRotator.Pitch = FMath::RoundHalfToEven(OutputRotator.Pitch / 90) * 90;
+	OutputRotator.Yaw = FMath::RoundHalfToEven(OutputRotator.Yaw / 90) * 90;
+	OutputRotator.Normalize();
+
+	// Convert from Rotator to Quaternion to Matrix and back To Rotator to ensure single representation of rotation value
+	FTransform MyTransform;
+	MyTransform.SetRotation(FQuat(OutputRotator));
+	FMatrix MyMatrix = MyTransform.ToMatrixNoScale();
+	OutputRotator= MyMatrix.Rotator();
+
+	// Ensure that -180 values are adjusted to 180
+	OutputRotator.Normalize();
+	OutputRotator.Roll = FMath::RoundHalfToEven(OutputRotator.Roll);
+	if (OutputRotator.Roll == -180) 
+	{ 
+		OutputRotator.Roll = 180; 
+	}
+	OutputRotator.Pitch = FMath::RoundHalfToEven(OutputRotator.Pitch);
+	if (OutputRotator.Pitch == -180) 
+	{ 
+		OutputRotator.Pitch = 180; 
+	}
+	OutputRotator.Yaw = FMath::RoundHalfToEven(OutputRotator.Yaw);
+	if (OutputRotator.Yaw == -180) 
+	{ 
+		OutputRotator.Yaw = 180; 
+	}
+
+	return OutputRotator;
+}
+
+void UWaveFunctionCollapseBPLibrary::DeriveModelFromActors(UPARAM(ref) const TArray<AActor*>& InputActors, 
+	UWaveFunctionCollapseModel* WFCModel, 
+	float TileSize, 
+	bool bIsBorderEmptyOption, 
+	bool bIsMinZFloorOption,
+	bool bUseUniformWeightDistribution,
+	bool bAutoDeriveZAxisRotationConstraints,
+	const TArray<FSoftObjectPath>& SpawnExclusionAssets,
+	const TArray<FSoftObjectPath>& IgnoreRotationAssets)
+{
+	// Check if there are valid actors
+	if (InputActors.IsEmpty())
+	{
+		return;
+	}
+
+	// Check if Model is valid
+	if (!WFCModel)
+	{
+		return;
+	}
+
+	// Initialize Variables
+	TMap<FIntVector, FWaveFunctionCollapseOption> PositionToOptionWithBorderMap;
+	FIntVector ResolutionWithBorder;
+
+	// Set TileSize
+	WFCModel->Modify();
+	WFCModel->TileSize = TileSize;
+
+	// Create all Empty Option adjacency constraints
+	WFCModel->AddConstraint(FWaveFunctionCollapseOption::EmptyOption, EWaveFunctionCollapseAdjacency::Front, FWaveFunctionCollapseOption::EmptyOption);
+	WFCModel->AddConstraint(FWaveFunctionCollapseOption::EmptyOption, EWaveFunctionCollapseAdjacency::Back, FWaveFunctionCollapseOption::EmptyOption);
+	WFCModel->AddConstraint(FWaveFunctionCollapseOption::EmptyOption, EWaveFunctionCollapseAdjacency::Left, FWaveFunctionCollapseOption::EmptyOption);
+	WFCModel->AddConstraint(FWaveFunctionCollapseOption::EmptyOption, EWaveFunctionCollapseAdjacency::Right, FWaveFunctionCollapseOption::EmptyOption);
+	WFCModel->AddConstraint(FWaveFunctionCollapseOption::EmptyOption, EWaveFunctionCollapseAdjacency::Up, FWaveFunctionCollapseOption::EmptyOption);
+	WFCModel->AddConstraint(FWaveFunctionCollapseOption::EmptyOption, EWaveFunctionCollapseAdjacency::Down, FWaveFunctionCollapseOption::EmptyOption);
+	
+	// Derive Grid from Actors
+	FIntVector MinPosition = FIntVector::ZeroValue;
+	FIntVector MaxPosition = FIntVector::ZeroValue;
+	FVector TmpOrigin = InputActors[0]->GetActorLocation(); // Set TmpOrigin based on first actor in the array
+	TMap<FIntVector, FWaveFunctionCollapseOption> TmpPositionToOptionMap;
+	for (int32 InputActorIndex = 0; InputActorIndex < InputActors.Num(); InputActorIndex++)
+	{
+		AStaticMeshActor* InputStaticMeshActor = Cast<AStaticMeshActor>(InputActors[InputActorIndex]);
+		if (InputStaticMeshActor)
+		{
+			UStaticMesh* InputStaticMesh = InputStaticMeshActor->GetStaticMeshComponent()->GetStaticMesh();
+			if (!InputStaticMesh)
+			{
+				UE_LOG(LogWFC, Display, TEXT("%s contains invalid StaticMesh, skipping."), *InputStaticMeshActor->GetFName().ToString());
+				continue;
+			}
+
+			// Find CurrentGridPosition
+			FVector CurrentGridPositionFloat = (InputStaticMeshActor->GetActorLocation() - TmpOrigin) / TileSize;
+			FIntVector CurrentGridPosition = FIntVector(FMath::RoundHalfFromZero(CurrentGridPositionFloat.X), FMath::RoundHalfFromZero(CurrentGridPositionFloat.Y), FMath::RoundHalfFromZero(CurrentGridPositionFloat.Z));
+			if (TmpPositionToOptionMap.Contains(CurrentGridPosition))
+			{
+				UE_LOG(LogWFC, Display, TEXT("%s is in an overlapping position, skipping."), *InputStaticMeshActor->GetFName().ToString());
+				continue;
+			}
+
+			// Find min/max positions
+			MinPosition.X = FMath::Min(MinPosition.X, CurrentGridPosition.X);
+			MinPosition.Y = FMath::Min(MinPosition.Y, CurrentGridPosition.Y);
+			MinPosition.Z = FMath::Min(MinPosition.Z, CurrentGridPosition.Z);
+			MaxPosition.X = FMath::Max(MaxPosition.X, CurrentGridPosition.X);
+			MaxPosition.Y = FMath::Max(MaxPosition.Y, CurrentGridPosition.Y);
+			MaxPosition.Z = FMath::Max(MaxPosition.Z, CurrentGridPosition.Z);
+
+			// Add to Temporary Position to Option Map
+			FRotator CurrentRotator = FRotator::ZeroRotator;
+			if (!IgnoreRotationAssets.Contains(FSoftObjectPath(InputStaticMesh->GetPathName())))
+			{
+				CurrentRotator = SanitizeRotator(InputStaticMeshActor->GetActorRotation());
+			}
+			FWaveFunctionCollapseOption CurrentOption(InputStaticMesh->GetPathName(), CurrentRotator, InputStaticMeshActor->GetActorScale3D());
+			TmpPositionToOptionMap.Add(CurrentGridPosition, CurrentOption);
+		}
+	}
+
+	// Re-order grid from 0,0,0 and add border
+	for (TPair<FIntVector, FWaveFunctionCollapseOption>& TmpPositionToOption : TmpPositionToOptionMap)
+	{
+		PositionToOptionWithBorderMap.Add(TmpPositionToOption.Key - MinPosition + FIntVector(1), TmpPositionToOption.Value);
+	}
+	TmpPositionToOptionMap.Empty();
+
+	// Set Resolution with Border
+	ResolutionWithBorder = MaxPosition - MinPosition + FIntVector(3);
+	
+	// Derive Constraints
+	for (TPair<FIntVector, FWaveFunctionCollapseOption>& PositionToOptionWithBorder : PositionToOptionWithBorderMap)
+	{
+		TMap<FIntVector, EWaveFunctionCollapseAdjacency> PositionToAdjacenciesMap = GetAdjacentPositions(PositionToOptionWithBorder.Key, ResolutionWithBorder);
+		for (TPair<FIntVector, EWaveFunctionCollapseAdjacency>& PositionToAdjacencies : PositionToAdjacenciesMap)
+		{
+			FIntVector PositionToCheck = PositionToAdjacencies.Key;
+
+			// if Adjacent Option exists, add constraint to model
+			if (PositionToOptionWithBorderMap.Contains(PositionToCheck))
+			{
+				WFCModel->AddConstraint(PositionToOptionWithBorder.Value, PositionToAdjacencies.Value, PositionToOptionWithBorderMap.FindRef(PositionToCheck));
+			}
+			// if MinZ and bIsMinZFloorOption, add Border Option constraint
+			else if (bIsMinZFloorOption && PositionToCheck.Z == 0)
+			{
+				WFCModel->AddConstraint(FWaveFunctionCollapseOption::BorderOption, GetOppositeAdjacency(PositionToAdjacencies.Value), PositionToOptionWithBorder.Value);
+			}
+			// if ExteriorBorder and bIsBorderEmptyOption, add Empty Option and Inverse constraints
+			else if (PositionToCheck.X == 0
+				|| PositionToCheck.Y == 0
+				|| PositionToCheck.Z == 0
+				|| PositionToCheck.X == ResolutionWithBorder.X - 1
+				|| PositionToCheck.Y == ResolutionWithBorder.Y - 1
+				|| PositionToCheck.Z == ResolutionWithBorder.Z - 1)
+			{
+				if (bIsBorderEmptyOption)
+				{
+					WFCModel->AddConstraint(PositionToOptionWithBorder.Value, PositionToAdjacencies.Value, FWaveFunctionCollapseOption::EmptyOption);
+					WFCModel->AddConstraint(FWaveFunctionCollapseOption::EmptyOption, GetOppositeAdjacency(PositionToAdjacencies.Value), PositionToOptionWithBorder.Value);
+				}
+			}
+			// otherwise it is an empty space, add Empty Option and Inverse constraints
+			else
+			{
+				WFCModel->AddConstraint(PositionToOptionWithBorder.Value, PositionToAdjacencies.Value, FWaveFunctionCollapseOption::EmptyOption);
+				WFCModel->AddConstraint(FWaveFunctionCollapseOption::EmptyOption, GetOppositeAdjacency(PositionToAdjacencies.Value), PositionToOptionWithBorder.Value);
+			}
+		}
+	}
+	PositionToOptionWithBorderMap.Empty();
+
+	// Auto Derive ZAxis Rotation Constraints
+	if (bAutoDeriveZAxisRotationConstraints)
+	{
+		TArray<FWaveFunctionCollapseOption> NewKeyOptions;
+		TArray<EWaveFunctionCollapseAdjacency> NewAdjacencies;
+		TArray<FWaveFunctionCollapseOption> NewAdjacentOptions;
+
+		for (const TPair<FWaveFunctionCollapseOption, FWaveFunctionCollapseAdjacencyToOptionsMap>& Constraint : WFCModel->Constraints)
+		{
+			for (const TPair<EWaveFunctionCollapseAdjacency, FWaveFunctionCollapseOptions>& AdjacencyToOptionsPair : Constraint.Value.AdjacencyToOptionsMap)
+			{
+				for (const FWaveFunctionCollapseOption& AdjacentOption : AdjacencyToOptionsPair.Value.Options)
+				{
+					FWaveFunctionCollapseOption NewKeyOption = Constraint.Key;
+					FWaveFunctionCollapseOption NewAdjacentOption = AdjacentOption;
+					EWaveFunctionCollapseAdjacency NewAdjacency = AdjacencyToOptionsPair.Key;
+
+					int32 KeyOptionRotationMultiplier = 0;
+					if (!(IgnoreRotationAssets.Contains(NewKeyOption.BaseObject)
+						|| NewKeyOption == FWaveFunctionCollapseOption::EmptyOption
+						|| NewKeyOption == FWaveFunctionCollapseOption::BorderOption
+						|| NewKeyOption == FWaveFunctionCollapseOption::VoidOption))
+					{
+						KeyOptionRotationMultiplier = 1;
+					}
+
+					int32 AdjacentOptionRotationMultiplier = 0;
+					if (!(IgnoreRotationAssets.Contains(NewAdjacentOption.BaseObject)
+						|| NewAdjacentOption == FWaveFunctionCollapseOption::EmptyOption
+						|| NewAdjacentOption == FWaveFunctionCollapseOption::BorderOption
+						|| NewAdjacentOption == FWaveFunctionCollapseOption::VoidOption))
+					{
+						AdjacentOptionRotationMultiplier = 1;
+					}
+
+					for (int32 RotationIncrement = 1; RotationIncrement <= 3; RotationIncrement++)
+					{
+						NewKeyOption.BaseRotator.Yaw += (90.0f * KeyOptionRotationMultiplier);
+						NewKeyOption.BaseRotator = SanitizeRotator(NewKeyOption.BaseRotator);
+						NewAdjacentOption.BaseRotator.Yaw += (90.0f * AdjacentOptionRotationMultiplier);
+						NewAdjacentOption.BaseRotator = SanitizeRotator(NewAdjacentOption.BaseRotator);
+						NewAdjacency = GetNextZAxisClockwiseAdjacency(NewAdjacency);
+
+						// Store new constraints
+						NewKeyOptions.Add(NewKeyOption);
+						NewAdjacencies.Add(NewAdjacency);
+						NewAdjacentOptions.Add(NewAdjacentOption);
+					}
+				}
+			}
+		}
+
+		// Add stored new constraints to the model
+		for (int32 Index = 0; Index < NewKeyOptions.Num(); Index++)
+		{
+			WFCModel->AddConstraint(NewKeyOptions[Index], NewAdjacencies[Index], NewAdjacentOptions[Index]);
+		}
+	}
+	
+	// If Empty->OptionA constraint and Empty<-OptionB exist, store OptionA->OptionB constraint
+	TArray<FWaveFunctionCollapseOption> EmptyVoidAdjacentKeyOptions;
+	TArray<EWaveFunctionCollapseAdjacency> EmptyVoidAdjacentAdjacencies;
+	TArray<FWaveFunctionCollapseOption> EmptyVoidAdjacentAdjacentOptions;
+	const FWaveFunctionCollapseAdjacencyToOptionsMap& EmptyOptionConstraints = WFCModel->Constraints.FindRef(FWaveFunctionCollapseOption::EmptyOption);
+	for (const TPair<EWaveFunctionCollapseAdjacency, FWaveFunctionCollapseOptions>& AdjacencyToOptionsPair : EmptyOptionConstraints.AdjacencyToOptionsMap)
+	{
+		for (const FWaveFunctionCollapseOption& EmptyAdjacentOption : AdjacencyToOptionsPair.Value.Options)
+		{
+			if (EmptyAdjacentOption == FWaveFunctionCollapseOption::EmptyOption)
+			{
+				continue;
+			}
+
+			if (const FWaveFunctionCollapseOptions* OppositeEmptyAdjacenctOptions = EmptyOptionConstraints.AdjacencyToOptionsMap.Find(GetOppositeAdjacency(AdjacencyToOptionsPair.Key)))
+			{
+				for (const FWaveFunctionCollapseOption& OppositeEmptyAdjacentOption : OppositeEmptyAdjacenctOptions->Options)
+				{
+					if (OppositeEmptyAdjacentOption == FWaveFunctionCollapseOption::EmptyOption)
+					{
+						continue;
+					}
+
+					EmptyVoidAdjacentKeyOptions.Add(EmptyAdjacentOption);
+					EmptyVoidAdjacentAdjacencies.Add(GetOppositeAdjacency(AdjacencyToOptionsPair.Key));
+					EmptyVoidAdjacentAdjacentOptions.Add(OppositeEmptyAdjacentOption);
+				}
+			}
+		}
+	}
+
+	// If Void->OptionA constraint and Void<-OptionB exist, store OptionA->OptionB constraint
+	const FWaveFunctionCollapseAdjacencyToOptionsMap& VoidOptionConstraints = WFCModel->Constraints.FindRef(FWaveFunctionCollapseOption::VoidOption);
+	for (const TPair<EWaveFunctionCollapseAdjacency, FWaveFunctionCollapseOptions>& AdjacencyToOptionsPair : VoidOptionConstraints.AdjacencyToOptionsMap)
+	{
+		for (const FWaveFunctionCollapseOption& VoidAdjacentOption : AdjacencyToOptionsPair.Value.Options)
+		{
+			if (VoidAdjacentOption == FWaveFunctionCollapseOption::VoidOption)
+			{
+				continue;
+			}
+
+			if (const FWaveFunctionCollapseOptions* OppositeVoidAdjacenctOptions = VoidOptionConstraints.AdjacencyToOptionsMap.Find(GetOppositeAdjacency(AdjacencyToOptionsPair.Key)))
+			{
+				for (const FWaveFunctionCollapseOption& OppositeVoidAdjacentOption : OppositeVoidAdjacenctOptions->Options)
+				{
+					if (VoidAdjacentOption == FWaveFunctionCollapseOption::VoidOption)
+					{
+						continue;
+					}
+
+					EmptyVoidAdjacentKeyOptions.Add(VoidAdjacentOption);
+					EmptyVoidAdjacentAdjacencies.Add(GetOppositeAdjacency(AdjacencyToOptionsPair.Key));
+					EmptyVoidAdjacentAdjacentOptions.Add(OppositeVoidAdjacentOption);
+				}
+			}
+		}
+	}
+
+	// Add stored new constraints to the model
+	for (int32 Index = 0; Index < EmptyVoidAdjacentKeyOptions.Num(); Index++)
+	{
+		WFCModel->AddConstraint(EmptyVoidAdjacentKeyOptions[Index], EmptyVoidAdjacentAdjacencies[Index], EmptyVoidAdjacentAdjacentOptions[Index]);
+	}
+	
+	// Set Contributions
+	if (bUseUniformWeightDistribution)
+	{
+		WFCModel->SetAllContributions(1);
+	}
+	WFCModel->SetOptionContribution(FWaveFunctionCollapseOption::BorderOption, 0);
+	WFCModel->SetWeightsFromContributions();
+
+	// Append to SpawnExlusion
+	for (FSoftObjectPath SpawnExclusionAsset : SpawnExclusionAssets)
+	{
+		if (SpawnExclusionAsset.IsValid())
+		{
+			WFCModel->SpawnExclusion.AddUnique(SpawnExclusionAsset);
+		}
+	}
+}
+
+bool UWaveFunctionCollapseBPLibrary::GetPositionToOptionMapFromActor(AActor* Actor, float TileSize, UPARAM(ref) TMap<FIntVector, FWaveFunctionCollapseOption>& PositionToOptionMap)
+{
+	// Check if Model is valid
+	if (!Actor)
+	{
+		UE_LOG(LogWFC, Display, TEXT("%s is an invalid Actor"), *Actor->GetFName().ToString());
+		return false;
+	}
+
+	// Derive Grid from Actors
+	FIntVector MinPosition = FIntVector::ZeroValue;
+	FIntVector MaxPosition = FIntVector::ZeroValue;
+	TMap<FIntVector, FWaveFunctionCollapseOption> TmpPositionToOptionMap;
+	
+	// Gather ISMComponents
+	TInlineComponentArray<UInstancedStaticMeshComponent*> ISMComponents;
+	Actor->GetComponents(ISMComponents);
+	for (int32 ISMComponentIndex = 0; ISMComponentIndex < ISMComponents.Num(); ISMComponentIndex++)
+	{
+		UInstancedStaticMeshComponent* ISMComponent = ISMComponents[ISMComponentIndex];
+		UStaticMesh* InputStaticMesh = ISMComponent->GetStaticMesh();
+		if (!InputStaticMesh)
+		{
+			UE_LOG(LogWFC, Display, TEXT("%s.%s contains invalid StaticMesh, skipping."), *Actor->GetFName().ToString(), *ISMComponent->GetFName().ToString());
+			continue;
+		}
+		else
+		{
+			for (int32 InstanceIndex = 0; InstanceIndex < ISMComponent->GetInstanceCount(); InstanceIndex++)
+			{
+				FTransform InstanceTransform;
+				ISMComponent->GetInstanceTransform(InstanceIndex, InstanceTransform);
+
+				// Find CurrentGridPosition
+				FVector CurrentGridPositionFloat = InstanceTransform.GetLocation() / TileSize;
+				FIntVector CurrentGridPosition = FIntVector(FMath::RoundHalfFromZero(CurrentGridPositionFloat.X), FMath::RoundHalfFromZero(CurrentGridPositionFloat.Y), FMath::RoundHalfFromZero(CurrentGridPositionFloat.Z));
+				if (TmpPositionToOptionMap.Contains(CurrentGridPosition))
+				{
+					UE_LOG(LogWFC, Display, TEXT("%s.%s Index %d is in an overlapping position, skipping."), *Actor->GetFName().ToString(), *ISMComponent->GetFName().ToString(), InstanceIndex);
+					continue;
+				}
+
+				// Find min/max positions
+				MinPosition.X = FMath::Min(MinPosition.X, CurrentGridPosition.X);
+				MinPosition.Y = FMath::Min(MinPosition.Y, CurrentGridPosition.Y);
+				MinPosition.Z = FMath::Min(MinPosition.Z, CurrentGridPosition.Z);
+				MaxPosition.X = FMath::Max(MaxPosition.X, CurrentGridPosition.X);
+				MaxPosition.Y = FMath::Max(MaxPosition.Y, CurrentGridPosition.Y);
+				MaxPosition.Z = FMath::Max(MaxPosition.Z, CurrentGridPosition.Z);
+
+				// Add to Temporary Position to Option Map
+				FWaveFunctionCollapseOption CurrentOption(InputStaticMesh->GetPathName(), InstanceTransform.Rotator(), InstanceTransform.GetScale3D());
+				TmpPositionToOptionMap.Add(CurrentGridPosition, CurrentOption);
+			}
+		}
+	}
+
+	// Re-order grid 
+	for (TPair<FIntVector, FWaveFunctionCollapseOption>& TmpPositionToOption : TmpPositionToOptionMap)
+	{
+		PositionToOptionMap.Add(TmpPositionToOption.Key - MinPosition, TmpPositionToOption.Value);
+	}
+	
+	return !PositionToOptionMap.IsEmpty();
+}
+
+FWaveFunctionCollapseOption UWaveFunctionCollapseBPLibrary::MakeEmptyOption()
+{
+	return FWaveFunctionCollapseOption::EmptyOption;
+}
+
+FWaveFunctionCollapseOption UWaveFunctionCollapseBPLibrary::MakeBorderOption()
+{
+	return FWaveFunctionCollapseOption::BorderOption;
+}
+
+FWaveFunctionCollapseOption UWaveFunctionCollapseBPLibrary::MakeVoidOption()
+{
+	return FWaveFunctionCollapseOption::VoidOption;
+}
diff --git a/Engine/Plugins/Experimental/WaveFunctionCollapse/Source/WaveFunctionCollapse/Private/WaveFunctionCollapseEditorCommands.cpp b/Engine/Plugins/Experimental/WaveFunctionCollapse/Source/WaveFunctionCollapse/Private/WaveFunctionCollapseEditorCommands.cpp
new file mode 100644
--- /dev/null
+++ b/Engine/Plugins/Experimental/WaveFunctionCollapse/Source/WaveFunctionCollapse/Private/WaveFunctionCollapseEditorCommands.cpp
@@ -0,0 +1,14 @@
+// Copyright Epic Games, Inc. All Rights Reserved.
+
+#include "WaveFunctionCollapseEditorCommands.h"
+#include "Framework/Commands/UICommandInfo.h"
+#include "Framework/Commands/InputChord.h"
+
+#define LOCTEXT_NAMESPACE "FWaveFunctionCollapseModule"
+
+void FWaveFunctionCollapseEditorCommands::RegisterCommands()
+{
+	UI_COMMAND(WaveFunctionCollapseWidget, "WaveFunctionCollapse", "Wave Function Collapse Tool", EUserInterfaceActionType::Button, FInputChord());
+}
+
+#undef LOCTEXT_NAMESPACE
\ No newline at end of file
diff --git a/Engine/Plugins/Experimental/WaveFunctionCollapse/Source/WaveFunctionCollapse/Private/WaveFunctionCollapseEditorCommands.h b/Engine/Plugins/Experimental/WaveFunctionCollapse/Source/WaveFunctionCollapse/Private/WaveFunctionCollapseEditorCommands.h
new file mode 100644
--- /dev/null
+++ b/Engine/Plugins/Experimental/WaveFunctionCollapse/Source/WaveFunctionCollapse/Private/WaveFunctionCollapseEditorCommands.h
@@ -0,0 +1,27 @@
+// Copyright Epic Games, Inc. All Rights Reserved.
+
+#pragma once
+
+#include "EditorStyleSet.h"
+#include "Framework/Commands/Commands.h"
+
+class FWaveFunctionCollapseEditorCommands : public TCommands<FWaveFunctionCollapseEditorCommands>
+{
+public:
+
+	FWaveFunctionCollapseEditorCommands()
+		: TCommands<FWaveFunctionCollapseEditorCommands>
+		(
+			TEXT("WaveFunctionCollapse"),
+			NSLOCTEXT("Contexts", "WaveFunctionCollapse", "WaveFunctionCollapse Plugin"),
+			NAME_None,
+			FEditorStyle::GetStyleSetName()
+		) 
+		{}
+
+	virtual void RegisterCommands() override;
+
+public:
+
+	TSharedPtr<FUICommandInfo> WaveFunctionCollapseWidget;
+};
\ No newline at end of file
diff --git a/Engine/Plugins/Experimental/WaveFunctionCollapse/Source/WaveFunctionCollapse/Private/WaveFunctionCollapseModel.cpp b/Engine/Plugins/Experimental/WaveFunctionCollapse/Source/WaveFunctionCollapse/Private/WaveFunctionCollapseModel.cpp
new file mode 100644
--- /dev/null
+++ b/Engine/Plugins/Experimental/WaveFunctionCollapse/Source/WaveFunctionCollapse/Private/WaveFunctionCollapseModel.cpp
@@ -0,0 +1,180 @@
+// Copyright Epic Games, Inc. All Rights Reserved.
+
+#include "WaveFunctionCollapseModel.h"
+
+const FWaveFunctionCollapseOption FWaveFunctionCollapseOption::EmptyOption(FString(TEXT("/WaveFunctionCollapse/Core/SpecialOptions/Option_Empty.Option_Empty")));
+const FWaveFunctionCollapseOption FWaveFunctionCollapseOption::BorderOption(FString(TEXT("/WaveFunctionCollapse/Core/SpecialOptions/Option_Border.Option_Border")));
+const FWaveFunctionCollapseOption FWaveFunctionCollapseOption::VoidOption(FString(TEXT("/WaveFunctionCollapse/Core/SpecialOptions/Option_Void.Option_Void")));
+
+void UWaveFunctionCollapseModel::AddConstraint(const FWaveFunctionCollapseOption& KeyOption, EWaveFunctionCollapseAdjacency Adjacency, const FWaveFunctionCollapseOption& AdjacentOption)
+{
+	Modify(true);
+	// If the KeyOption key exists in Constraints
+	if (FWaveFunctionCollapseAdjacencyToOptionsMap* AdjacencyToOptionsMap = Constraints.Find(KeyOption))
+	{
+		// If the Adjacency key exists in the AdjacencyToOptionsMap
+		if (FWaveFunctionCollapseOptions* Options = AdjacencyToOptionsMap->AdjacencyToOptionsMap.Find(Adjacency))
+		{
+			// If the Options array does not contain the Option, create it
+			if (!Options->Options.Contains(AdjacentOption))
+			{
+				Options->Options.Add(AdjacentOption);
+				AdjacencyToOptionsMap->Contribution += 1;
+			}
+		}
+		else // Create the Option and add it to the AdjacencyToOptionsMap
+		{
+			FWaveFunctionCollapseOptions NewOptions;
+			NewOptions.Options.Add(AdjacentOption);
+			AdjacencyToOptionsMap->AdjacencyToOptionsMap.Add(Adjacency, NewOptions);
+			AdjacencyToOptionsMap->Contribution += 1;
+		}
+	}
+	else // Create the Option, add it to the AdjacencyToOptionsMap, and add it the Constraints map
+	{
+		FWaveFunctionCollapseOptions NewOptions;
+		NewOptions.Options.Add(AdjacentOption);
+		FWaveFunctionCollapseAdjacencyToOptionsMap NewAdjacencyToOptionsMap;
+		NewAdjacencyToOptionsMap.AdjacencyToOptionsMap.Add(Adjacency, NewOptions);
+		Constraints.Add(KeyOption, NewAdjacencyToOptionsMap);
+	}
+	
+}
+
+FWaveFunctionCollapseOptions UWaveFunctionCollapseModel::GetOptions(const FWaveFunctionCollapseOption& KeyOption, EWaveFunctionCollapseAdjacency Adjacency) const
+{
+	// If KeyOption key exists
+	if (const FWaveFunctionCollapseAdjacencyToOptionsMap* AdjacencyToOptionsMap = Constraints.Find(KeyOption))
+	{
+		// If Adjacency key exists, return FoundOptions
+		if (const FWaveFunctionCollapseOptions* FoundOptions = AdjacencyToOptionsMap->AdjacencyToOptionsMap.Find(Adjacency))
+		{
+			return *FoundOptions;
+		}
+	}
+	// If nothing is found above, return empty
+	return FWaveFunctionCollapseOptions();
+}
+
+void UWaveFunctionCollapseModel::SetWeightsFromContributions()
+{
+	if (Constraints.IsEmpty())
+	{
+		return;
+	}
+
+	Modify(true);
+	int32 SumOfContributions = 0;
+	for (TPair<FWaveFunctionCollapseOption, FWaveFunctionCollapseAdjacencyToOptionsMap>& Constraint : Constraints)
+	{
+		SumOfContributions += Constraint.Value.Contribution;
+	}
+
+	for (TPair<FWaveFunctionCollapseOption, FWaveFunctionCollapseAdjacencyToOptionsMap>& Constraint : Constraints)
+	{
+		Constraint.Value.Weight = float(Constraint.Value.Contribution) / SumOfContributions;
+	}
+}
+
+void UWaveFunctionCollapseModel::SetAllWeights(float Weight)
+{
+	if (Constraints.IsEmpty())
+	{
+		return;
+	}
+
+	Modify(true);
+	for (TPair<FWaveFunctionCollapseOption, FWaveFunctionCollapseAdjacencyToOptionsMap>& Constraint : Constraints)
+	{
+		Constraint.Value.Weight = Weight;
+	}
+}
+
+void UWaveFunctionCollapseModel::SetAllContributions(int32 Contribution)
+{
+	if (Constraints.IsEmpty())
+	{
+		return;
+	}
+
+	Modify(true);
+	for (TPair<FWaveFunctionCollapseOption, FWaveFunctionCollapseAdjacencyToOptionsMap>& Constraint : Constraints)
+	{
+		Constraint.Value.Contribution = Contribution;
+	}
+}
+
+void UWaveFunctionCollapseModel::SetOptionContribution(const FWaveFunctionCollapseOption& Option, int32 Contribution)
+{
+	if (Constraints.Contains(Option))
+	{
+		Modify(true);
+		FWaveFunctionCollapseAdjacencyToOptionsMap& AdjacencyToOptionsMap = *Constraints.Find(Option);
+		AdjacencyToOptionsMap.Contribution = Contribution;
+	}
+}
+
+float UWaveFunctionCollapseModel::GetOptionWeight(const FWaveFunctionCollapseOption& Option) const
+{
+	if (const FWaveFunctionCollapseAdjacencyToOptionsMap* Constraint = Constraints.Find(Option))
+	{
+		return Constraint->Weight;
+	}
+	else
+	{
+		return 0;
+	}
+}
+
+int32 UWaveFunctionCollapseModel::GetOptionContribution(const FWaveFunctionCollapseOption& Option) const
+{
+	if (const FWaveFunctionCollapseAdjacencyToOptionsMap* Constraint = Constraints.Find(Option))
+	{
+		return Constraint->Contribution;
+	}
+	else
+	{
+		return 0;
+	}
+}
+
+int32 UWaveFunctionCollapseModel::GetConstraintCount() const
+{
+	int32 ConstraintCount = 0;
+	for (const TPair<FWaveFunctionCollapseOption, FWaveFunctionCollapseAdjacencyToOptionsMap>& Constraint : Constraints)
+	{
+		for (const TPair<EWaveFunctionCollapseAdjacency, FWaveFunctionCollapseOptions>& AdjacencyToOptions : Constraint.Value.AdjacencyToOptionsMap)
+		{
+			ConstraintCount += AdjacencyToOptions.Value.Options.Num();
+		}
+	}
+	return ConstraintCount;
+}
+
+void UWaveFunctionCollapseModel::SwapMeshes(TMap<UStaticMesh*, UStaticMesh*> SourceToTargetMeshMap)
+{
+	Modify(true);
+	for (TPair<UStaticMesh*, UStaticMesh*>& SourceToTargetMesh : SourceToTargetMeshMap)
+	{
+		UStaticMesh* SourceMesh = SourceToTargetMesh.Key;
+		UStaticMesh* TargetMesh = SourceToTargetMesh.Value;
+		for (TPair<FWaveFunctionCollapseOption, FWaveFunctionCollapseAdjacencyToOptionsMap>& Constraint : Constraints)
+		{
+			if (Constraint.Key.BaseObject == SourceMesh)
+			{
+				Constraint.Key.BaseObject = TargetMesh;
+			}
+
+			for (TPair<EWaveFunctionCollapseAdjacency, FWaveFunctionCollapseOptions>& AdjacencyToOptions : Constraint.Value.AdjacencyToOptionsMap)
+			{
+				for (FWaveFunctionCollapseOption& Option : AdjacencyToOptions.Value.Options)
+				{
+					if (Option.BaseObject == SourceMesh)
+					{
+						Option.BaseObject = TargetMesh;
+					}
+				}
+			}
+		}
+	}
+}
\ No newline at end of file
diff --git a/Engine/Plugins/Experimental/WaveFunctionCollapse/Source/WaveFunctionCollapse/Private/WaveFunctionCollapseSubsystem.cpp b/Engine/Plugins/Experimental/WaveFunctionCollapse/Source/WaveFunctionCollapse/Private/WaveFunctionCollapseSubsystem.cpp
new file mode 100644
--- /dev/null
+++ b/Engine/Plugins/Experimental/WaveFunctionCollapse/Source/WaveFunctionCollapse/Private/WaveFunctionCollapseSubsystem.cpp
@@ -0,0 +1,857 @@
+// Copyright Epic Games, Inc. All Rights Reserved.
+
+#include "WaveFunctionCollapseSubsystem.h"
+#include "WaveFunctionCollapseBPLibrary.h"
+#include "Math/UnrealMathUtility.h"
+#include "Math/NumericLimits.h"
+#include "Math/Vector.h"
+#include "Engine/World.h"
+#include "Components/InstancedStaticMeshComponent.h"
+#include "Subsystems/EditorActorSubsystem.h"
+#include "Kismet2/ComponentEditorUtils.h"
+#include "Editor.h"
+#include "Editor/EditorEngine.h"
+
+DEFINE_LOG_CATEGORY(LogWFC);
+
+AActor* UWaveFunctionCollapseSubsystem::Collapse(int32 TryCount /* = 1 */, int32 RandomSeed /* = 0 */)
+{
+	if (!WFCModel)
+	{
+		UE_LOG(LogWFC, Error, TEXT("Invalid WFC Model"));
+		return nullptr;
+	}
+	
+	UE_LOG(LogWFC, Display, TEXT("Starting WFC - Model: %s, Resolution %dx%dx%d"), *WFCModel->GetFName().ToString(), Resolution.X, Resolution.Y, Resolution.Z);
+
+	// Determinism settings
+	int32 ChosenRandomSeed = (RandomSeed != 0 ? RandomSeed : FMath::RandRange(1, TNumericLimits<int32>::Max()));
+
+	int32 ArrayReserveValue = Resolution.X * Resolution.Y * Resolution.Z;
+	TArray<FWaveFunctionCollapseTile> Tiles;
+	TArray<int32> RemainingTiles;
+	TMap<int32, FWaveFunctionCollapseQueueElement> ObservationQueue;
+	Tiles.Reserve(ArrayReserveValue);
+	RemainingTiles.Reserve(ArrayReserveValue);
+
+	InitializeWFC(Tiles, RemainingTiles);
+	
+	bool bSuccessfulSolve = false;
+
+	if (TryCount > 1)
+	{
+		//Copy Original Initialized tiles
+		TArray<FWaveFunctionCollapseTile> TilesCopy = Tiles;
+		TArray<int32> RemainingTilesCopy = RemainingTiles;
+
+		int32 CurrentTry = 1;
+		bSuccessfulSolve = ObservationPropagation(Tiles, RemainingTiles, ObservationQueue, ChosenRandomSeed);
+		FRandomStream RandomStream(ChosenRandomSeed);
+		while (!bSuccessfulSolve && CurrentTry<TryCount)
+		{
+			CurrentTry += 1;
+			UE_LOG(LogWFC, Warning, TEXT("Failed with Seed Value: %d. Trying again.  Attempt number: %d"), ChosenRandomSeed, CurrentTry);
+			ChosenRandomSeed = RandomStream.RandRange(1, TNumericLimits<int32>::Max());
+			
+			// Start from Original Initialized tiles
+			Tiles = TilesCopy;
+			RemainingTiles = RemainingTilesCopy;
+			ObservationQueue.Empty();
+			bSuccessfulSolve = ObservationPropagation(Tiles, RemainingTiles, ObservationQueue, ChosenRandomSeed);
+		}
+	}
+	else if (TryCount == 1)
+	{
+		bSuccessfulSolve = ObservationPropagation(Tiles, RemainingTiles, ObservationQueue, ChosenRandomSeed);
+	}
+	else
+	{
+		UE_LOG(LogWFC, Error, TEXT("Invalid TryCount on Collapse: %d"), TryCount);
+		return nullptr;
+	}
+
+	// if Successful, Spawn Actor
+	if (bSuccessfulSolve)
+	{
+		AActor* SpawnedActor = SpawnActorFromTiles(Tiles);
+		UE_LOG(LogWFC, Display, TEXT("Success! Seed Value: %d. Spawned Actor: %s"), ChosenRandomSeed, *SpawnedActor->GetActorLabel());
+		return SpawnedActor;
+	}
+	else
+	{
+		UE_LOG(LogWFC, Error, TEXT("Failed after %d tries."), TryCount);
+		return nullptr;
+	}
+}
+
+void UWaveFunctionCollapseSubsystem::InitializeWFC(TArray<FWaveFunctionCollapseTile>& Tiles, TArray<int32>& RemainingTiles)
+{
+	FWaveFunctionCollapseTile InitialTile;
+	int32 SwapIndex = 0;
+
+	if (BuildInitialTile(InitialTile))
+	{
+		float MinEntropy = InitialTile.ShannonEntropy;
+		for (int32 Z = 0;Z < Resolution.Z; Z++)
+		{
+			for (int32 Y = 0;Y < Resolution.Y; Y++)
+			{
+				for (int32 X = 0;X < Resolution.X; X++)
+				{
+					// Pre-populate with starter tiles
+					if (FWaveFunctionCollapseOption* StarterOption = StarterOptions.Find(FIntVector(X, Y, Z)))
+					{
+						FWaveFunctionCollapseTile StarterTile;
+						StarterTile.RemainingOptions.Add(*StarterOption);
+						StarterTile.ShannonEntropy = UWaveFunctionCollapseBPLibrary::CalculateShannonEntropy(StarterTile.RemainingOptions, WFCModel);
+						Tiles.Add(StarterTile);
+						RemainingTiles.Add(UWaveFunctionCollapseBPLibrary::PositionAsIndex(FIntVector(X, Y, Z), Resolution));
+						
+						// swap lower entropy tile to the beginning of RemainingTiles
+						if (StarterTile.ShannonEntropy < MinEntropy)
+						{
+							RemainingTiles.Swap(0, RemainingTiles.Num() - 1);
+							MinEntropy = StarterTile.ShannonEntropy;
+							SwapIndex = 0;
+						}
+						
+						// else, swap min entropy tile with the previous min entropy index+1
+						else if (StarterTile.ShannonEntropy == MinEntropy && StarterTile.ShannonEntropy != InitialTile.ShannonEntropy)
+						{
+							SwapIndex += 1;
+							RemainingTiles.Swap(SwapIndex, RemainingTiles.Num() - 1);
+						}
+					}
+					
+					// Pre-populate with border tiles
+					else if (IsPositionInnerBorder(FIntVector(X, Y, Z))
+						&& (bUseEmptyBorder || WFCModel->Constraints.Contains(FWaveFunctionCollapseOption::BorderOption)))
+					{
+						FWaveFunctionCollapseTile BorderTile;
+						BorderTile.RemainingOptions = GetInnerBorderOptions(FIntVector(X, Y, Z), InitialTile.RemainingOptions);
+						BorderTile.ShannonEntropy = UWaveFunctionCollapseBPLibrary::CalculateShannonEntropy(BorderTile.RemainingOptions, WFCModel);
+						Tiles.Add(BorderTile);
+						RemainingTiles.Add(UWaveFunctionCollapseBPLibrary::PositionAsIndex(FIntVector(X, Y, Z), Resolution));
+
+						// swap lower entropy tile to the beginning of RemainingTiles
+						if (BorderTile.ShannonEntropy < MinEntropy)
+						{
+							RemainingTiles.Swap(0, RemainingTiles.Num() - 1);
+							MinEntropy = BorderTile.ShannonEntropy;
+							SwapIndex = 0;
+						}
+						
+						// else, swap min entropy tile with the previous min entropy index+1
+						else if (BorderTile.ShannonEntropy ==  MinEntropy && BorderTile.ShannonEntropy != InitialTile.ShannonEntropy)
+						{
+							SwapIndex += 1;
+							RemainingTiles.Swap(SwapIndex, RemainingTiles.Num() - 1);
+						}
+					}
+					
+					// Fill the rest with initial tiles
+					else
+					{
+						Tiles.Add(InitialTile);
+						RemainingTiles.Add(UWaveFunctionCollapseBPLibrary::PositionAsIndex(FIntVector(X, Y, Z), Resolution));
+					}
+				}
+			}
+		}
+		StarterOptions.Empty();
+	}
+	else
+	{
+		UE_LOG(LogWFC, Error, TEXT("Could not create Initial Tile from Model"));
+	}
+
+}
+
+bool UWaveFunctionCollapseSubsystem::BuildInitialTile(FWaveFunctionCollapseTile& InitialTile)
+{
+	TArray<FWaveFunctionCollapseOption> InitialOptions;
+	for (const TPair< FWaveFunctionCollapseOption, FWaveFunctionCollapseAdjacencyToOptionsMap>& Constraint : WFCModel->Constraints)
+	{
+		if (Constraint.Key.BaseObject != FWaveFunctionCollapseOption::BorderOption.BaseObject)
+		{
+			InitialOptions.Add(Constraint.Key);
+		}
+	}
+
+	if (!InitialOptions.IsEmpty())
+	{
+		InitialTile.RemainingOptions = InitialOptions;
+		InitialTile.ShannonEntropy = UWaveFunctionCollapseBPLibrary::CalculateShannonEntropy(InitialOptions, WFCModel);
+		return true;
+	}
+	else
+	{
+		return false;
+	}
+}
+
+TArray<FWaveFunctionCollapseOption> UWaveFunctionCollapseSubsystem::GetInnerBorderOptions(FIntVector Position, const TArray<FWaveFunctionCollapseOption>& InitialOptions)
+{
+	TArray<FWaveFunctionCollapseOption> InnerBorderOptions;
+	TArray<FWaveFunctionCollapseOption> InnerBorderOptionsToRemove;
+	InnerBorderOptions = InitialOptions;
+	
+	// gather options to remove
+	if (Position.X == 0)
+	{
+		GatherInnerBorderOptionsToRemove(EWaveFunctionCollapseAdjacency::Front, InnerBorderOptions, InnerBorderOptionsToRemove);
+	}
+	if (Position.X == Resolution.X - 1)
+	{
+		GatherInnerBorderOptionsToRemove(EWaveFunctionCollapseAdjacency::Back, InnerBorderOptions, InnerBorderOptionsToRemove);
+	}
+	if (Position.Y == 0)
+	{
+		GatherInnerBorderOptionsToRemove(EWaveFunctionCollapseAdjacency::Right, InnerBorderOptions, InnerBorderOptionsToRemove);
+	}
+	if (Position.Y == Resolution.Y - 1)
+	{
+		GatherInnerBorderOptionsToRemove(EWaveFunctionCollapseAdjacency::Left, InnerBorderOptions, InnerBorderOptionsToRemove);
+	}
+	if (Position.Z == 0)
+	{
+		GatherInnerBorderOptionsToRemove(EWaveFunctionCollapseAdjacency::Up, InnerBorderOptions, InnerBorderOptionsToRemove);
+	}
+	if (Position.Z == Resolution.Z - 1)
+	{
+		GatherInnerBorderOptionsToRemove(EWaveFunctionCollapseAdjacency::Down, InnerBorderOptions, InnerBorderOptionsToRemove);
+	}
+	
+	//remove options
+	if (!InnerBorderOptionsToRemove.IsEmpty())
+	{
+		for (FWaveFunctionCollapseOption& RemoveThisOption : InnerBorderOptionsToRemove)
+		{
+			InnerBorderOptions.RemoveSingleSwap(RemoveThisOption, false);
+		}
+		InnerBorderOptions.Shrink();
+	}
+	
+	return InnerBorderOptions;
+}
+
+void UWaveFunctionCollapseSubsystem::GatherInnerBorderOptionsToRemove(EWaveFunctionCollapseAdjacency Adjacency, const TArray<FWaveFunctionCollapseOption>& InitialOptions, TArray<FWaveFunctionCollapseOption>& OutBorderOptionsToRemove)
+{
+	bool bFoundBorderOptions = false;
+	for (const FWaveFunctionCollapseOption& InitialOption : InitialOptions)
+	{
+		// if border option exists in the model, use it
+		if (FWaveFunctionCollapseAdjacencyToOptionsMap* FoundBorderAdjacencyToOptionsMap = WFCModel->Constraints.Find(FWaveFunctionCollapseOption::BorderOption))
+		{
+			if (FWaveFunctionCollapseOptions* FoundBorderOptions = FoundBorderAdjacencyToOptionsMap->AdjacencyToOptionsMap.Find(Adjacency))
+			{
+				if (!FoundBorderOptions->Options.Contains(InitialOption))
+				{
+					OutBorderOptionsToRemove.AddUnique(InitialOption);
+					bFoundBorderOptions = true;
+				}
+			}
+		}
+		
+		// else, if useEmptyBorder, use empty option
+		if (bUseEmptyBorder && !bFoundBorderOptions)
+		{
+			if (FWaveFunctionCollapseAdjacencyToOptionsMap* FoundEmptyAdjacencyToOptionsMap = WFCModel->Constraints.Find(FWaveFunctionCollapseOption::EmptyOption))
+			{
+				if (FWaveFunctionCollapseOptions* FoundEmptyOptions = FoundEmptyAdjacencyToOptionsMap->AdjacencyToOptionsMap.Find(Adjacency))
+				{
+					if (!FoundEmptyOptions->Options.Contains(InitialOption))
+					{
+						OutBorderOptionsToRemove.AddUnique(InitialOption);
+					}
+				}
+			}
+		}
+	}
+}
+
+bool UWaveFunctionCollapseSubsystem::IsPositionInnerBorder(FIntVector Position)
+{
+	return (Position.X == 0
+		|| Position.Y == 0
+		|| Position.Z == 0
+		|| Position.X == Resolution.X - 1
+		|| Position.Y == Resolution.Y - 1
+		|| Position.Z == Resolution.Z - 1);
+}
+
+bool UWaveFunctionCollapseSubsystem::Observe(TArray<FWaveFunctionCollapseTile>& Tiles, 
+	TArray<int32>& RemainingTiles, 
+	TMap<int32, FWaveFunctionCollapseQueueElement>& ObservationQueue,
+	int32 RandomSeed)
+{
+	float MinEntropy = 0;
+	int32 LastSameMinEntropyIndex = 0;
+	int32 SelectedMinEntropyIndex = 0;
+	int32 MinEntropyIndex = 0;
+	FRandomStream RandomStream(RandomSeed);
+
+	// Find MinEntropy Tile Indices
+	if (RemainingTiles.Num() > 1)
+	{
+		for (int32 index = 0; index<RemainingTiles.Num(); index++)
+		{
+			if (index == 0)
+			{
+				MinEntropy = Tiles[RemainingTiles[index]].ShannonEntropy;
+			}
+			else
+			{
+				if (Tiles[RemainingTiles[index]].ShannonEntropy > MinEntropy)
+				{
+					break;
+				}
+				else
+				{
+					LastSameMinEntropyIndex += 1;
+				}
+			}
+		}
+		SelectedMinEntropyIndex = RandomStream.RandRange(0, LastSameMinEntropyIndex);
+		MinEntropyIndex = RemainingTiles[SelectedMinEntropyIndex];
+	}
+	else
+	{
+		MinEntropyIndex = RemainingTiles[0];
+	}
+
+	// Rand Selection of Weighted Options using Cumulative Density
+	TArray<float> CumulativeDensity;
+	CumulativeDensity.Reserve(Tiles[MinEntropyIndex].RemainingOptions.Num());
+	float CumulativeWeight = 0;
+	for (FWaveFunctionCollapseOption& Option : Tiles[MinEntropyIndex].RemainingOptions)
+	{
+		CumulativeWeight += WFCModel->Constraints.Find(Option)->Weight;
+		CumulativeDensity.Add(CumulativeWeight);
+	}
+	
+	int32 SelectedOptionIndex = 0;
+	float RandomDensity = RandomStream.FRandRange(0.0f, CumulativeDensity.Last());
+	for (int32 Index = 0; Index < CumulativeDensity.Num(); Index++)
+	{
+		if (CumulativeDensity[Index] > RandomDensity)
+		{
+			SelectedOptionIndex = Index;
+			break;
+		}
+	}
+
+	// Make Selection
+	Tiles[MinEntropyIndex] = FWaveFunctionCollapseTile(Tiles[MinEntropyIndex].RemainingOptions[SelectedOptionIndex], TNumericLimits<float>::Max());
+
+	if (SelectedMinEntropyIndex != LastSameMinEntropyIndex)
+	{
+		RemainingTiles.Swap(SelectedMinEntropyIndex, LastSameMinEntropyIndex);
+	}
+	RemainingTiles.RemoveAtSwap(LastSameMinEntropyIndex);
+
+	if (!RemainingTiles.IsEmpty())
+	{
+		// if MinEntropy has changed after removal, find new MinEntropy and swap to front of array
+		if (Tiles[RemainingTiles[0]].ShannonEntropy != MinEntropy)
+		{
+			int32 SwapToIndex = 0;
+			for (int32 index = 0; index < RemainingTiles.Num(); index++)
+			{
+				if (index == 0)
+				{
+					MinEntropy = Tiles[RemainingTiles[index]].ShannonEntropy;
+				}
+				else
+				{
+					if (Tiles[RemainingTiles[index]].ShannonEntropy < MinEntropy)
+					{
+						SwapToIndex = 0;
+						MinEntropy = Tiles[RemainingTiles[index]].ShannonEntropy;
+						RemainingTiles.Swap(SwapToIndex, index);
+					}
+					else if (Tiles[RemainingTiles[index]].ShannonEntropy == MinEntropy)
+					{
+						SwapToIndex += 1;
+						RemainingTiles.Swap(SwapToIndex, index);
+					}
+				}
+			}
+		}
+
+		// Add Adjacent Tile Indices to Queue
+		AddAdjacentIndicesToQueue(MinEntropyIndex, RemainingTiles, ObservationQueue);
+
+		// Continue To Propagation
+		return true;
+	}
+	else
+	{
+		// Do Not Continue to Propagation
+		return false;
+	}
+}
+
+void UWaveFunctionCollapseSubsystem::AddAdjacentIndicesToQueue(int32 CenterIndex, const TArray<int32>& RemainingTiles, TMap<int32,FWaveFunctionCollapseQueueElement>& OutQueue)
+{
+	FIntVector Position = UWaveFunctionCollapseBPLibrary::IndexAsPosition(CenterIndex, Resolution);
+	if (Position.X + 1 < Resolution.X && RemainingTiles.Contains(UWaveFunctionCollapseBPLibrary::PositionAsIndex(Position + FIntVector(1, 0, 0), Resolution)))
+	{ 
+		OutQueue.Add(UWaveFunctionCollapseBPLibrary::PositionAsIndex(Position + FIntVector(1, 0, 0), Resolution),
+			FWaveFunctionCollapseQueueElement(CenterIndex, EWaveFunctionCollapseAdjacency::Front));
+	}
+	if (Position.X - 1 >= 0 && RemainingTiles.Contains(UWaveFunctionCollapseBPLibrary::PositionAsIndex(Position + FIntVector(-1, 0, 0), Resolution)))
+	{ 
+		OutQueue.Add(UWaveFunctionCollapseBPLibrary::PositionAsIndex(Position + FIntVector(-1, 0, 0), Resolution),
+			FWaveFunctionCollapseQueueElement(CenterIndex, EWaveFunctionCollapseAdjacency::Back));
+	}
+	if (Position.Y + 1 < Resolution.Y && RemainingTiles.Contains(UWaveFunctionCollapseBPLibrary::PositionAsIndex(Position + FIntVector(0, 1, 0), Resolution)))
+	{ 
+		OutQueue.Add(UWaveFunctionCollapseBPLibrary::PositionAsIndex(Position + FIntVector(0, 1, 0), Resolution),
+			FWaveFunctionCollapseQueueElement(CenterIndex, EWaveFunctionCollapseAdjacency::Right));
+	}
+	if (Position.Y - 1 >= 0 && RemainingTiles.Contains(UWaveFunctionCollapseBPLibrary::PositionAsIndex(Position + FIntVector(0, -1, 0), Resolution)))
+	{ 
+		OutQueue.Add(UWaveFunctionCollapseBPLibrary::PositionAsIndex(Position + FIntVector(0, -1, 0), Resolution),
+			FWaveFunctionCollapseQueueElement(CenterIndex, EWaveFunctionCollapseAdjacency::Left));
+	}
+	if (Position.Z + 1 < Resolution.Z && RemainingTiles.Contains(UWaveFunctionCollapseBPLibrary::PositionAsIndex(Position + FIntVector(0, 0, 1), Resolution)))
+	{ 
+		OutQueue.Add(UWaveFunctionCollapseBPLibrary::PositionAsIndex(Position + FIntVector(0, 0, 1), Resolution),
+			FWaveFunctionCollapseQueueElement(CenterIndex, EWaveFunctionCollapseAdjacency::Up));
+	}
+	if (Position.Z - 1 >= 0 && RemainingTiles.Contains(UWaveFunctionCollapseBPLibrary::PositionAsIndex(Position + FIntVector(0, 0, -1), Resolution)))
+	{ 
+		OutQueue.Add(UWaveFunctionCollapseBPLibrary::PositionAsIndex(Position + FIntVector(0, 0, -1), Resolution),
+			FWaveFunctionCollapseQueueElement(CenterIndex, EWaveFunctionCollapseAdjacency::Down));
+	}
+}
+
+bool UWaveFunctionCollapseSubsystem::Propagate(TArray<FWaveFunctionCollapseTile>& Tiles, 
+	TArray<int32>& RemainingTiles, 
+	TMap<int32, FWaveFunctionCollapseQueueElement>& ObservationQueue, 
+	int32& PropagationCount)
+{
+	TMap<int32, FWaveFunctionCollapseQueueElement> PropagationQueue;
+
+	while (!ObservationQueue.IsEmpty())
+	{
+		for (TPair<int32, FWaveFunctionCollapseQueueElement>& ObservationAdjacenctElement : ObservationQueue)
+		{
+			// Make sure the tile to check is still a valid remaining tile
+			if (!RemainingTiles.Contains(ObservationAdjacenctElement.Key))
+			{
+				continue;
+			}
+			
+			TArray<FWaveFunctionCollapseOption>& ObservationRemainingOptions = Tiles[ObservationAdjacenctElement.Key].RemainingOptions;
+
+			// Get check against options
+			TArray<FWaveFunctionCollapseOption> OptionsToCheckAgainst;
+			OptionsToCheckAgainst.Reserve(WFCModel->Constraints.Num());
+			for (FWaveFunctionCollapseOption& CenterOption : Tiles[ObservationAdjacenctElement.Value.CenterObjectIndex].RemainingOptions)
+			{
+				for (FWaveFunctionCollapseOption& Option : WFCModel->Constraints.FindRef(CenterOption).AdjacencyToOptionsMap.FindRef(ObservationAdjacenctElement.Value.Adjacency).Options)
+				{
+					OptionsToCheckAgainst.AddUnique(Option);
+				}
+			}
+				
+			// Accumulate New Remaining Options
+			bool bAddToPropagationQueue = false;
+			TArray<FWaveFunctionCollapseOption> TmpRemainingOptionsAccumulation;
+			TmpRemainingOptionsAccumulation.Reserve(ObservationRemainingOptions.Num());
+			for (FWaveFunctionCollapseOption& ObservationRemainingOption : ObservationRemainingOptions)
+			{
+				if (OptionsToCheckAgainst.Contains(ObservationRemainingOption))
+				{
+					TmpRemainingOptionsAccumulation.AddUnique(ObservationRemainingOption);
+				}
+				else
+				{
+					bAddToPropagationQueue = true;
+				}
+			}
+				
+			// If Remaining Options have changed
+			if (bAddToPropagationQueue)
+			{
+				if (!TmpRemainingOptionsAccumulation.IsEmpty())
+				{
+					AddAdjacentIndicesToQueue(ObservationAdjacenctElement.Key, RemainingTiles, PropagationQueue);
+
+					// Update Tile with new options
+					float MinEntropy = Tiles[RemainingTiles[0]].ShannonEntropy;
+					float NewEntropy = UWaveFunctionCollapseBPLibrary::CalculateShannonEntropy(TmpRemainingOptionsAccumulation, WFCModel);
+					int32 CurrentRemainingTileIndex;
+						
+					// If NewEntropy is <= MinEntropy, add to front of Remaining Tiles
+					if (NewEntropy < MinEntropy)
+					{
+						MinEntropy = NewEntropy;
+						CurrentRemainingTileIndex = RemainingTiles.Find(ObservationAdjacenctElement.Key);
+						if (CurrentRemainingTileIndex != 0)
+						{
+							RemainingTiles.Swap(0, CurrentRemainingTileIndex);
+						}
+					}
+					else if (NewEntropy == MinEntropy)
+					{
+						CurrentRemainingTileIndex = RemainingTiles.Find(ObservationAdjacenctElement.Key);
+						for (int32 Index = 1; Index < RemainingTiles.Num(); Index++)
+						{
+							if (MinEntropy != Tiles[RemainingTiles[Index]].ShannonEntropy)
+							{
+								if (CurrentRemainingTileIndex != Index)
+								{
+									RemainingTiles.Swap(Index, CurrentRemainingTileIndex);
+								}
+								break;
+							}
+						}
+					}
+						
+					Tiles[ObservationAdjacenctElement.Key] = FWaveFunctionCollapseTile(TmpRemainingOptionsAccumulation, NewEntropy);
+				}
+				else
+				{
+					// Encountered Contradiction
+					UE_LOG(LogWFC, Error, TEXT("Encountered Contradiction on Index %d"), ObservationAdjacenctElement.Key);
+					return false;
+				}
+			}
+		}
+
+		ObservationQueue = PropagationQueue;
+		if (!PropagationQueue.IsEmpty())
+		{
+			PropagationCount += 1;
+			PropagationQueue.Reset();
+		}
+	}
+
+	return true;
+}
+
+bool UWaveFunctionCollapseSubsystem::ObservationPropagation(TArray<FWaveFunctionCollapseTile>& Tiles, 
+	TArray<int32>& RemainingTiles,
+	TMap<int32, FWaveFunctionCollapseQueueElement>& ObservationQueue,
+	int32 RandomSeed)
+{
+	int32 PropagationCount = 1;
+	int32 MutatedRandomSeed = RandomSeed;
+	
+	while (Observe(Tiles, RemainingTiles, ObservationQueue, MutatedRandomSeed))
+	{
+		if (!Propagate(Tiles, RemainingTiles, ObservationQueue, PropagationCount))
+		{
+			return false;
+		}
+
+		// Mutate Seed
+		MutatedRandomSeed--;
+	}
+
+	// Check if all tiles in the solve are non-spawnable
+	return !AreAllTilesNonSpawnable(Tiles);
+}
+
+UActorComponent* UWaveFunctionCollapseSubsystem::AddNamedInstanceComponent(AActor* Actor, TSubclassOf<UActorComponent> ComponentClass, FName ComponentName)
+{
+	Actor->Modify();
+	// Assign Unique Name
+	int32 Counter = 1;
+	FName ComponentInstanceName = ComponentName;
+	while (!FComponentEditorUtils::IsComponentNameAvailable(ComponentInstanceName.ToString(), Actor))
+	{
+		ComponentInstanceName = FName(*FString::Printf(TEXT("%s_%d"), *ComponentName.ToString(), Counter++));
+	}
+	UActorComponent* InstanceComponent = NewObject<UActorComponent>(Actor, ComponentClass, ComponentInstanceName, RF_Transactional);
+	if (InstanceComponent)
+	{
+		Actor->AddInstanceComponent(InstanceComponent);
+		Actor->FinishAddComponent(InstanceComponent, false, FTransform::Identity);
+		Actor->RerunConstructionScripts();
+	}
+	return InstanceComponent;
+}
+
+AActor* UWaveFunctionCollapseSubsystem::SpawnActorFromTiles(const TArray<FWaveFunctionCollapseTile>& Tiles)
+{
+	UEditorActorSubsystem* EditorActorSubsystem = GEditor->GetEditorSubsystem<UEditorActorSubsystem>();
+	// Spawn Actor
+	AActor* SpawnedActor = EditorActorSubsystem->SpawnActorFromClass(AActor::StaticClass(), FVector::ZeroVector);
+	SpawnedActor->GetRootComponent()->SetMobility(EComponentMobility::Static);
+	SpawnedActor->SetActorLocationAndRotation(OriginLocation, Orientation);
+	FActorLabelUtilities::SetActorLabelUnique(SpawnedActor, WFCModel->GetFName().ToString());
+
+	// Create Components
+	TMap<FSoftObjectPath, UInstancedStaticMeshComponent*> BaseObjectToISM;
+	for (int32 index = 0; index < Tiles.Num(); index++)
+	{
+		if (Tiles[index].RemainingOptions.Num() != 1)
+		{
+			continue;
+		}
+
+		// check for empty or void options
+		FSoftObjectPath BaseObject = Tiles[index].RemainingOptions[0].BaseObject;
+		if (BaseObject == FWaveFunctionCollapseOption::EmptyOption.BaseObject
+			|| BaseObject == FWaveFunctionCollapseOption::VoidOption.BaseObject)
+		{
+			continue;
+		}
+
+		// Check the SpawnExclusion array
+		if (WFCModel->SpawnExclusion.Contains(BaseObject))
+		{
+			continue;
+		}
+
+		UObject* LoadedObject = BaseObject.TryLoad();
+		if (LoadedObject)
+		{
+			FRotator BaseRotator = Tiles[index].RemainingOptions[0].BaseRotator;
+			FVector BaseScale3D = Tiles[index].RemainingOptions[0].BaseScale3D;
+			FVector PositionOffset = FVector(WFCModel->TileSize * 0.5f);
+			FVector TilePosition = (FVector(UWaveFunctionCollapseBPLibrary::IndexAsPosition(index, Resolution)) * WFCModel->TileSize) + PositionOffset;
+
+			// Static meshes are handled with ISM Components
+			if (UStaticMesh* LoadedStaticMesh = Cast<UStaticMesh>(LoadedObject))
+			{
+				UInstancedStaticMeshComponent* ISMComponent;
+				if (UInstancedStaticMeshComponent** FoundISMComponentPtr = BaseObjectToISM.Find(BaseObject))
+				{
+					ISMComponent = *FoundISMComponentPtr;
+				}
+				else
+				{
+					ISMComponent = Cast<UInstancedStaticMeshComponent>(UWaveFunctionCollapseSubsystem::AddNamedInstanceComponent(SpawnedActor, UInstancedStaticMeshComponent::StaticClass(), LoadedObject->GetFName()));
+					BaseObjectToISM.Add(BaseObject, ISMComponent);
+				}
+				ISMComponent->SetStaticMesh(LoadedStaticMesh);
+				ISMComponent->SetMobility(EComponentMobility::Static);
+				ISMComponent->AddInstance(FTransform(BaseRotator, TilePosition, BaseScale3D));
+			}
+			// Blueprints are handled with ChildActorComponents
+			else if (UBlueprint* LoadedBlueprint = Cast<UBlueprint>(LoadedObject))
+			{
+				// if BP is placeable
+				if (!(LoadedBlueprint->GetClass()->HasAnyClassFlags(CLASS_NotPlaceable | CLASS_Deprecated | CLASS_Abstract))
+					&& LoadedBlueprint->GetClass()->IsChildOf(AActor::StaticClass()))
+				{
+					UChildActorComponent* ChildActorComponent = Cast<UChildActorComponent>(UWaveFunctionCollapseSubsystem::AddNamedInstanceComponent(SpawnedActor, UChildActorComponent::StaticClass(), LoadedObject->GetFName()));
+					ChildActorComponent->SetChildActorClass(LoadedBlueprint->GetClass());
+					ChildActorComponent->SetRelativeLocation(TilePosition);
+					ChildActorComponent->SetRelativeRotation(BaseRotator);
+					ChildActorComponent->SetRelativeScale3D(BaseScale3D);
+				}
+			}
+			else
+			{
+				UE_LOG(LogWFC, Warning, TEXT("Invalid Type, skipping: %s"), *BaseObject.ToString());
+			}
+		}
+		else
+		{
+			UE_LOG(LogWFC, Warning, TEXT("Unable to load object, skipping: %s"), *BaseObject.ToString());
+		}
+	}
+
+	return SpawnedActor;
+}
+
+bool UWaveFunctionCollapseSubsystem::AreAllTilesNonSpawnable(const TArray<FWaveFunctionCollapseTile>& Tiles)
+{
+	bool bAllTilesAreNonSpawnable = true;
+	for (int32 index = 0; index < Tiles.Num(); index++)
+	{
+		if (Tiles[index].RemainingOptions.Num() == 1)
+		{
+			FSoftObjectPath BaseObject = Tiles[index].RemainingOptions[0].BaseObject;
+			if (!(BaseObject == FWaveFunctionCollapseOption::EmptyOption.BaseObject
+				|| BaseObject == FWaveFunctionCollapseOption::VoidOption.BaseObject
+				|| WFCModel->SpawnExclusion.Contains(BaseObject)))
+			{
+				bAllTilesAreNonSpawnable = false;
+				break;
+			}
+		}
+	}
+	return bAllTilesAreNonSpawnable;
+}
+
+void UWaveFunctionCollapseSubsystem::DeriveGridFromTransformBounds(const TArray<FTransform>& Transforms)
+{
+	if (Transforms.IsEmpty())
+	{
+		UE_LOG(LogWFC, Error, TEXT("Empty Transform Array."));
+		return;
+	}
+
+	FVector TransformPivot;
+	FVector AxisAlignedPoint = FVector::ZeroVector;
+	FVector MinBound = FVector::ZeroVector;
+	FVector MaxBound = FVector::ZeroVector;
+	TMap<FIntVector, int32> PositionToPointCount;
+	float TransformOrientation = Transforms[0].GetRotation().Rotator().Yaw;
+
+	float CosOrientation = FMath::Cos(PI / (180.f) * -TransformOrientation);
+	float SinOrientation = FMath::Sin(PI / (180.f) * -TransformOrientation);
+	float UnitScale = 1.0f;
+
+	for (int32 index = 0; index < Transforms.Num(); index++)
+	{
+		if (index == 0)
+		{
+			TransformPivot = Transforms[index].GetLocation() * UnitScale;
+		}
+		else
+		{
+			FVector OffsetPoint = Transforms[index].GetLocation() * UnitScale - TransformPivot;
+			AxisAlignedPoint.X = (OffsetPoint.X * CosOrientation) - (OffsetPoint.Y * SinOrientation);
+			AxisAlignedPoint.Y = (OffsetPoint.X * SinOrientation) + (OffsetPoint.Y * CosOrientation);
+			AxisAlignedPoint.Z = OffsetPoint.Z;
+			MinBound.X = FMath::Min(AxisAlignedPoint.X, MinBound.X);
+			MinBound.Y = FMath::Min(AxisAlignedPoint.Y, MinBound.Y);
+			MinBound.Z = FMath::Min(AxisAlignedPoint.Z, MinBound.Z);
+			MaxBound.X = FMath::Max(AxisAlignedPoint.X, MaxBound.X);
+			MaxBound.Y = FMath::Max(AxisAlignedPoint.Y, MaxBound.Y);
+			MaxBound.Z = FMath::Max(AxisAlignedPoint.Z, MaxBound.Z);
+		}
+		
+		FIntVector PointPosition;
+		PointPosition.X = (int32)FMath::RoundHalfFromZero(AxisAlignedPoint.X / WFCModel->TileSize);
+		PointPosition.Y = (int32)FMath::RoundHalfFromZero(AxisAlignedPoint.Y / WFCModel->TileSize);
+		PointPosition.Z = (int32)FMath::RoundHalfFromZero(AxisAlignedPoint.Z / WFCModel->TileSize);
+		if (int32* PointCount = PositionToPointCount.Find(PointPosition))
+		{
+			PositionToPointCount.Add(PointPosition, *PointCount + 1);
+		}
+		else
+		{
+			PositionToPointCount.Add(PointPosition, 1);
+		}
+	}
+
+	// Set WFC Resolution
+	Resolution.X = (int32)FMath::RoundHalfFromZero(MaxBound.X / WFCModel->TileSize) - (int32)FMath::RoundHalfFromZero(MinBound.X / WFCModel->TileSize) + 1;
+	Resolution.Y = (int32)FMath::RoundHalfFromZero(MaxBound.Y / WFCModel->TileSize) - (int32)FMath::RoundHalfFromZero(MinBound.Y / WFCModel->TileSize) + 1;
+	Resolution.Z = (int32)FMath::RoundHalfFromZero(MaxBound.Z / WFCModel->TileSize) - (int32)FMath::RoundHalfFromZero(MinBound.Z / WFCModel->TileSize) + 1;
+
+	// Set WFC OriginLocation
+	FVector ReorientedMinPoint;
+	ReorientedMinPoint.X = ((MinBound.X - (WFCModel->TileSize * 0.5f)) * CosOrientation) - ((MinBound.Y - (WFCModel->TileSize * 0.5f)) * -SinOrientation);
+	ReorientedMinPoint.Y = ((MinBound.X - (WFCModel->TileSize * 0.5f)) * -SinOrientation) + ((MinBound.Y - (WFCModel->TileSize * 0.5f)) * CosOrientation);
+	ReorientedMinPoint.Z = MinBound.Z - (WFCModel->TileSize * 0.5f);
+	OriginLocation = ReorientedMinPoint + TransformPivot;
+
+	// Set WFC Orientation
+	Orientation = FRotator(0, TransformOrientation, 0);
+
+	// Set WFC Starter Options
+	FIntVector PositionOffset;
+	PositionOffset.X = (int32)FMath::RoundHalfFromZero(MinBound.X / WFCModel->TileSize);
+	PositionOffset.Y = (int32)FMath::RoundHalfFromZero(MinBound.Y / WFCModel->TileSize);
+	PositionOffset.Z = (int32)FMath::RoundHalfFromZero(MinBound.Z / WFCModel->TileSize);
+	StarterOptions.Empty();
+	for (int32 z = 0;z < Resolution.Z; z++)
+	{
+		for (int32 y = 0;y < Resolution.Y; y++)
+		{
+			for (int32 x = 0;x < Resolution.X; x++)
+			{
+				if (!PositionToPointCount.Contains(FIntVector(x + PositionOffset.X, y + PositionOffset.Y, z + PositionOffset.Z)))
+				{
+					StarterOptions.Add(FIntVector(x,y,z), FWaveFunctionCollapseOption::EmptyOption);
+				}
+			}
+		}
+	}
+}
+
+void UWaveFunctionCollapseSubsystem::DeriveGridFromTransforms(const TArray<FTransform>& Transforms)
+{
+	if (Transforms.IsEmpty())
+	{
+		UE_LOG(LogWFC, Error, TEXT("Empty Transform Array."));
+		return;
+	}
+
+	FVector TransformPivot;
+	FVector AxisAlignedPoint = FVector::ZeroVector;
+	FVector MinBound = FVector::ZeroVector;
+	FVector MaxBound = FVector::ZeroVector;
+	TArray<FIntVector> PointPositions;
+	float TransformOrientation = Transforms[0].GetRotation().Rotator().Yaw;
+
+	float CosOrientation = FMath::Cos(PI / (180.f) * -TransformOrientation);
+	float SinOrientation = FMath::Sin(PI / (180.f) * -TransformOrientation);
+	float UnitScale = 1.0f;
+	
+	for (int32 index = 0; index < Transforms.Num(); index++)
+	{
+		if (index == 0)
+		{
+			TransformPivot = (Transforms[index].GetLocation() * UnitScale);
+		}
+		else
+		{
+			FVector LocalPoint = Transforms[index].GetLocation() * UnitScale - TransformPivot;
+			AxisAlignedPoint.X = (LocalPoint.X * CosOrientation) - (LocalPoint.Y * SinOrientation);
+			AxisAlignedPoint.Y = (LocalPoint.X * SinOrientation) + (LocalPoint.Y * CosOrientation);
+			AxisAlignedPoint.Z = LocalPoint.Z;
+			MinBound.X = FMath::Min(AxisAlignedPoint.X, MinBound.X);
+			MinBound.Y = FMath::Min(AxisAlignedPoint.Y, MinBound.Y);
+			MinBound.Z = FMath::Min(AxisAlignedPoint.Z, MinBound.Z);
+			MaxBound.X = FMath::Max(AxisAlignedPoint.X, MaxBound.X);
+			MaxBound.Y = FMath::Max(AxisAlignedPoint.Y, MaxBound.Y);
+			MaxBound.Z = FMath::Max(AxisAlignedPoint.Z, MaxBound.Z);
+		}
+		
+		FIntVector PointPosition;
+		PointPosition.X = (int32)FMath::RoundHalfFromZero(AxisAlignedPoint.X / WFCModel->TileSize);
+		PointPosition.Y = (int32)FMath::RoundHalfFromZero(AxisAlignedPoint.Y / WFCModel->TileSize);
+		PointPosition.Z = (int32)FMath::RoundHalfFromZero(AxisAlignedPoint.Z / WFCModel->TileSize);
+		PointPositions.AddUnique(PointPosition);
+	}
+
+	// Set WFC Resolution
+	Resolution.X = (int32)FMath::RoundHalfFromZero(MaxBound.X / WFCModel->TileSize) - (int32)FMath::RoundHalfFromZero(MinBound.X / WFCModel->TileSize) + 1;
+	Resolution.Y = (int32)FMath::RoundHalfFromZero(MaxBound.Y / WFCModel->TileSize) - (int32)FMath::RoundHalfFromZero(MinBound.Y / WFCModel->TileSize) + 1;
+	Resolution.Z = (int32)FMath::RoundHalfFromZero(MaxBound.Z / WFCModel->TileSize) - (int32)FMath::RoundHalfFromZero(MinBound.Z / WFCModel->TileSize) + 1;
+
+	// Set WFC OriginLocation
+	FVector ReorientedMinPoint;
+	ReorientedMinPoint.X = ((MinBound.X - (WFCModel->TileSize * 0.5f)) * CosOrientation) - ((MinBound.Y - (WFCModel->TileSize * 0.5f)) * -SinOrientation);
+	ReorientedMinPoint.Y = ((MinBound.X - (WFCModel->TileSize * 0.5f)) * -SinOrientation) + ((MinBound.Y - (WFCModel->TileSize * 0.5f)) * CosOrientation);
+	ReorientedMinPoint.Z = MinBound.Z - (WFCModel->TileSize * 0.5f);
+	OriginLocation = ReorientedMinPoint + TransformPivot;
+	
+	// Set WFC Orientation
+	Orientation = FRotator(0, TransformOrientation, 0);
+
+	// Set WFC Starter Options
+	FIntVector PositionOffset;
+	PositionOffset.X = (int32)FMath::RoundHalfFromZero(MinBound.X / WFCModel->TileSize);
+	PositionOffset.Y = (int32)FMath::RoundHalfFromZero(MinBound.Y / WFCModel->TileSize);
+	PositionOffset.Z = (int32)FMath::RoundHalfFromZero(MinBound.Z / WFCModel->TileSize);
+	StarterOptions.Empty();
+	for (int32 z = 0;z < Resolution.Z; z++)
+	{
+		for (int32 y = 0;y < Resolution.Y; y++)
+		{
+			for (int32 x = 0;x < Resolution.X; x++)
+			{
+				if (!PointPositions.Contains(FIntVector(x + PositionOffset.X, y + PositionOffset.Y, z + PositionOffset.Z)))
+				{
+					StarterOptions.Add(FIntVector(x, y, z), FWaveFunctionCollapseOption::EmptyOption);
+				}
+			}
+		}
+	}
+	for (TPair<FIntVector,FWaveFunctionCollapseOption> StarterOption : StarterOptions)
+	{
+		UE_LOG(LogWFC, Display, TEXT("StartOption at: %d,%d,%d"), StarterOption.Key.X, StarterOption.Key.Y, StarterOption.Key.Z);
+	}
+}
\ No newline at end of file
diff --git a/Engine/Plugins/Experimental/WaveFunctionCollapse/Source/WaveFunctionCollapse/Public/WaveFunctionCollapse.h b/Engine/Plugins/Experimental/WaveFunctionCollapse/Source/WaveFunctionCollapse/Public/WaveFunctionCollapse.h
new file mode 100644
--- /dev/null
+++ b/Engine/Plugins/Experimental/WaveFunctionCollapse/Source/WaveFunctionCollapse/Public/WaveFunctionCollapse.h
@@ -0,0 +1,24 @@
+// Copyright Epic Games, Inc. All Rights Reserved.
+
+#pragma once
+
+#include "Modules/ModuleManager.h"
+
+class FMenuBuilder;
+class FReply;
+
+class FWaveFunctionCollapseModule : public IModuleInterface
+{
+public:
+
+	/** IModuleInterface implementation */
+	virtual void StartupModule() override;
+	virtual void ShutdownModule() override;
+
+private:
+
+	void AddMenuEntry(FMenuBuilder& MenuBuilder);
+	void WaveFunctionCollapseUI();
+
+	TSharedPtr<class FUICommandList> PluginCommands;
+};
\ No newline at end of file
diff --git a/Engine/Plugins/Experimental/WaveFunctionCollapse/Source/WaveFunctionCollapse/Public/WaveFunctionCollapseBPLibrary.h b/Engine/Plugins/Experimental/WaveFunctionCollapse/Source/WaveFunctionCollapse/Public/WaveFunctionCollapseBPLibrary.h
new file mode 100644
--- /dev/null
+++ b/Engine/Plugins/Experimental/WaveFunctionCollapse/Source/WaveFunctionCollapse/Public/WaveFunctionCollapseBPLibrary.h
@@ -0,0 +1,193 @@
+// Copyright Epic Games, Inc. All Rights Reserved.
+
+#pragma once
+
+#include "Kismet/BlueprintFunctionLibrary.h"
+#include "WaveFunctionCollapseModel.h"
+#include "WaveFunctionCollapseBPLibrary.generated.h"
+
+USTRUCT(BlueprintType)
+struct WAVEFUNCTIONCOLLAPSE_API FWaveFunctionCollapseNeighbor
+{
+	GENERATED_BODY()
+
+	UPROPERTY(BlueprintReadWrite, EditAnywhere, Category = "WaveFunctionCollapse")
+	EWaveFunctionCollapseAdjacency Adjacency = EWaveFunctionCollapseAdjacency::Front;
+
+	UPROPERTY(BlueprintReadWrite, EditAnywhere, Category = "WaveFunctionCollapse")
+	int32 Step = 1;
+
+	UPROPERTY(BlueprintReadWrite, EditAnywhere, Category = "WaveFunctionCollapse", meta = (AllowedClasses = "StaticMesh, Blueprint"))
+	FSoftObjectPath NeighborObject;
+};
+
+USTRUCT(BlueprintType)
+struct WAVEFUNCTIONCOLLAPSE_API FWaveFunctionCollapseNeighborRule
+{
+	GENERATED_BODY()
+
+	UPROPERTY(BlueprintReadWrite, EditAnywhere, Category = "WaveFunctionCollapse", meta = (AllowedClasses = "StaticMesh, Blueprint"))
+	FSoftObjectPath KeyObject;
+
+	UPROPERTY(BlueprintReadWrite, EditAnywhere, Category = "WaveFunctionCollapse")
+	TArray<FWaveFunctionCollapseNeighbor> Neighbors;
+
+	UPROPERTY(BlueprintReadWrite, EditAnywhere, Category = "WaveFunctionCollapse", meta = (AllowedClasses = "StaticMesh, Blueprint"))
+	TArray<FSoftObjectPath> SpawnObjects;
+
+	UPROPERTY(BlueprintReadWrite, EditAnywhere, Category = "WaveFunctionCollapse")
+	FTransform SpawnRelativeTransform;
+
+	/**
+	* SpawnChance 1 = 100 % chance to spawn, SpawnChance 0 = 0 % chance to spawn
+	*/ 
+	UPROPERTY(BlueprintReadWrite, EditAnywhere, Category = "WaveFunctionCollapse", meta = (ClampMin = "0", ClampMax = "1"))
+	float SpawnChance = 1.0f;
+};
+
+UCLASS()
+class WAVEFUNCTIONCOLLAPSE_API UWaveFunctionCollapseBPLibrary : public UBlueprintFunctionLibrary
+{
+	GENERATED_UCLASS_BODY()
+
+	/**
+	* Calculates Shannon Entropy from an array of options and a given model
+	* @param Options Array of options
+	* @param WFCModel WaveFunctionCollapseModel that stores weights for options 
+	*/
+	UFUNCTION(BlueprintCallable, Category = "WaveFunctionCollapse")
+	static float CalculateShannonEntropy(const TArray<FWaveFunctionCollapseOption>& Options, UWaveFunctionCollapseModel* WFCModel);
+
+	/**
+	* Convert 3D grid position to 2D array index
+	* @param Position 3D grid position
+	* @param Resolution Grid resolution
+	* @return 2D array index
+	*/
+	UFUNCTION(BlueprintPure, Category = "WaveFunctionCollapse")
+	static int32 PositionAsIndex(FIntVector Position, FIntVector Resolution);
+
+	/**
+	* Convert 2D array index to 3D grid position
+	* @param Index 2D array index
+	* @param Resolution
+	* @return 3D grid position
+	*/
+	UFUNCTION(BlueprintPure, Category = "WaveFunctionCollapse")
+	static FIntVector IndexAsPosition(int32 Index, FIntVector Resolution);
+
+	/**
+	* Builds the initial tile which adds every unique option in a model to its RemainingOptions array and calculates its entropy
+	* @param WFCModel WaveFunctionCollapseModel
+	*/
+	UFUNCTION(BlueprintCallable, Category = "WaveFunctionCollapse", meta = (ScriptMethod))
+	static FWaveFunctionCollapseTile BuildInitialTile(UWaveFunctionCollapseModel* WFCModel);
+
+	/**
+	* Get adjacent indices of a given index and its adjacency
+	* @param Index 2D array index
+	* @param Resolution
+	* @return Map of valid adjacent indices and its adjacency in relation to the given index
+	*/
+	UFUNCTION(BlueprintCallable, Category = "WaveFunctionCollapse")
+	static TMap<int32, EWaveFunctionCollapseAdjacency> GetAdjacentIndices(int32 Index, FIntVector Resolution);
+
+	/**
+	* Get adjacent positions of a given position and its adjacency
+	* @param Position 3D grid position
+	* @param Resolution Grid resolution
+	* @return Map of valid adjacent positions and its adjacency in relation to the given position
+	*/
+	UFUNCTION(BlueprintCallable, Category = "WaveFunctionCollapse")
+	static TMap<FIntVector, EWaveFunctionCollapseAdjacency> GetAdjacentPositions(FIntVector Position, FIntVector Resolution);
+
+	/**
+	* Is the option contained in the given options array
+	* @param Option Option to check for
+	* @param Options Array of options
+	*/
+	UFUNCTION(BlueprintPure, Category = "WaveFunctionCollapse")
+	static bool IsOptionContained(const FWaveFunctionCollapseOption& Option, const TArray<FWaveFunctionCollapseOption>& Options);
+
+	/**
+	* Get the opposite adjacency for a given adjacency.  For example GetOppositeAdjacency(Front) will return Back.
+	* @param Adjacency Adjacency direction
+	*/
+	UFUNCTION(BlueprintPure, Category = "WaveFunctionCollapse")
+	static EWaveFunctionCollapseAdjacency GetOppositeAdjacency(EWaveFunctionCollapseAdjacency Adjacency);
+	
+	/**
+	* Get the next adjacency in clockwise direction on a Z-axis for a given adjacency.
+	* For example GetNextZAxisClockwiseAdjacency(Front) will return Right.
+	* Up or Down will return the original adjacency.
+	* @param Adjacency Adjacency direction
+	*/
+	UFUNCTION(BlueprintPure, Category = "WaveFunctionCollapse")
+	static EWaveFunctionCollapseAdjacency GetNextZAxisClockwiseAdjacency(EWaveFunctionCollapseAdjacency Adjacency);
+
+	/**
+	* Add an entry to an AdjacencyToOptionsMap
+	*/
+	UFUNCTION(BlueprintCallable, Category = "WaveFunctionCollapse")
+	static void AddToAdjacencyToOptionsMap(UPARAM(ref) FWaveFunctionCollapseAdjacencyToOptionsMap& AdjacencyToOptionsMap, EWaveFunctionCollapseAdjacency Adjacency, FWaveFunctionCollapseOption Option);
+
+	UFUNCTION(BlueprintPure, Category = "WaveFunctionCollapse")
+	static bool IsSoftObjPathEqual(const FSoftObjectPath& SoftObjectPathA, const FSoftObjectPath& SoftObjectPathB);
+
+	/**
+	* Converts Rotator to Matrix and back to sanitize multiple representations of the same rotation to a common Rotator value
+	*/
+	UFUNCTION(BlueprintPure, Category = "WaveFunctionCollapse")
+	static FRotator SanitizeRotator(FRotator Rotator);
+
+	/**
+	* Derive constraints from a given layout of actors and append them to a model
+	* @param Actors array of actors to evaluate
+	* @param WFCModel to add constraints to
+	* @param TileSize distance between tiles
+	* @param bIsBorderEmptyOption should the border be considered EmptyOption
+	* @param bIsMinZFloorOption should the minimum Z actors be considered floor options (nothing can go below it)
+	* @param bAutoDeriveZAxisRotationConstraints should it auto derive z-axis rotational variants
+	* @param SpawnExclusionAssets assets to exclude when spawning
+	* @param IgnoreRotationAssets assets to ignore rotation variants
+	*/
+	UFUNCTION(BlueprintCallable, Category = "WaveFunctionCollapse")
+	static void DeriveModelFromActors(UPARAM(ref) const TArray<AActor*>& Actors, 
+		UWaveFunctionCollapseModel* WFCModel, 
+		float TileSize, 
+		bool bIsBorderEmptyOption, 
+		bool bIsMinZFloorOption, 
+		bool bUseUniformWeightDistribution,
+		bool bAutoDeriveZAxisRotationConstraints,
+		const TArray<FSoftObjectPath>& SpawnExclusionAssets,
+		const TArray<FSoftObjectPath>& IgnoreRotationAssets);
+
+	/**
+	* Get PositionToOptionsMap from a given actor that has ISM components.
+	* This is useful when you want to derive neighboring tile data from a WFC-solved actor to be used for post processing.
+	* This will only evaluate ISM components.
+	* @param Actor Actor with ISM components
+	* @param TileSize distance between tiles
+	* @param PositionToOptionMap 
+	*/
+	UFUNCTION(BlueprintCallable, Category = "WaveFunctionCollapse")
+	static bool GetPositionToOptionMapFromActor(AActor* Actor, float TileSize, UPARAM(ref) TMap<FIntVector, FWaveFunctionCollapseOption>& PositionToOptionMap);
+
+	/**
+	* Make FWaveFunctionCollapseOption of type: EmptyOption
+	*/
+	UFUNCTION(BlueprintPure, Category = "WaveFunctionCollapse")
+	static FWaveFunctionCollapseOption MakeEmptyOption();
+
+	/**
+	* Make FWaveFunctionCollapseOption of type: BorderOption
+	*/
+	UFUNCTION(BlueprintPure, Category = "WaveFunctionCollapse")
+	static FWaveFunctionCollapseOption MakeBorderOption();
+	
+	/**
+	* Make FWaveFunctionCollapseOption of type: VoidOption
+	*/
+	UFUNCTION(BlueprintPure, Category = "WaveFunctionCollapse")
+	static FWaveFunctionCollapseOption MakeVoidOption();
+};
diff --git a/Engine/Plugins/Experimental/WaveFunctionCollapse/Source/WaveFunctionCollapse/Public/WaveFunctionCollapseModel.h b/Engine/Plugins/Experimental/WaveFunctionCollapse/Source/WaveFunctionCollapse/Public/WaveFunctionCollapseModel.h
new file mode 100644
--- /dev/null
+++ b/Engine/Plugins/Experimental/WaveFunctionCollapse/Source/WaveFunctionCollapse/Public/WaveFunctionCollapseModel.h
@@ -0,0 +1,257 @@
+// Copyright Epic Games, Inc. All Rights Reserved.
+
+#pragma once
+
+#include "Engine/DataAsset.h"
+#include "Engine/StaticMesh.h"
+#include "WaveFunctionCollapseModel.generated.h"
+
+UENUM(BlueprintType)
+enum class EWaveFunctionCollapseAdjacency : uint8
+{
+	Front	UMETA(DisplayName = "X+ Front"),
+	Back	UMETA(DisplayName = "X- Back"),
+	Right	UMETA(DisplayName = "Y+ Right"),
+	Left	UMETA(DisplayName = "Y- Left"),
+	Up		UMETA(DisplayName = "Z+ Up"),
+	Down	UMETA(DisplayName = "Z- Down")
+};
+
+/**
+* Base Option Struct which holds an object, its orientation and scale
+*/
+USTRUCT(BlueprintType)
+struct WAVEFUNCTIONCOLLAPSE_API FWaveFunctionCollapseOption
+{
+	GENERATED_BODY()
+
+	UPROPERTY(BlueprintReadWrite, EditAnywhere, Category = "WaveFunctionCollapse", meta = (AllowedClasses = "StaticMesh, Blueprint"))
+	FSoftObjectPath BaseObject;
+
+	UPROPERTY(BlueprintReadWrite, EditAnywhere, Category = "WaveFunctionCollapse")
+	FRotator BaseRotator = FRotator::ZeroRotator;
+
+	UPROPERTY(BlueprintReadWrite, EditAnywhere, Category = "WaveFunctionCollapse")
+	FVector BaseScale3D = FVector::OneVector;
+
+	FWaveFunctionCollapseOption() = default;
+
+	FWaveFunctionCollapseOption(FString Object, FRotator Rotator, FVector Scale3d)
+	{
+		BaseObject = FSoftObjectPath(Object);
+		BaseRotator = Rotator;
+		BaseScale3D = Scale3d;
+	}
+
+	FWaveFunctionCollapseOption(FString Object)
+	{
+		BaseObject = FSoftObjectPath(Object);
+	}
+
+	static const FWaveFunctionCollapseOption EmptyOption;
+	static const FWaveFunctionCollapseOption BorderOption;
+	static const FWaveFunctionCollapseOption VoidOption;
+
+	friend uint32 GetTypeHash(FWaveFunctionCollapseOption Output)
+	{
+		uint32 OutputHash;
+		OutputHash = HashCombine(GetTypeHash(Output.BaseRotator.Vector()), GetTypeHash(Output.BaseScale3D));
+		OutputHash = HashCombine(OutputHash, GetTypeHash(Output.BaseObject));
+		return OutputHash;
+	}
+
+	bool operator==(const FWaveFunctionCollapseOption& Rhs) const
+	{
+		return BaseObject == Rhs.BaseObject && BaseRotator.Equals(Rhs.BaseRotator) && BaseScale3D.Equals(Rhs.BaseScale3D);
+	}
+
+	bool operator!=(const FWaveFunctionCollapseOption& Rhs) const
+	{
+		return BaseObject != Rhs.BaseObject || !BaseRotator.Equals(Rhs.BaseRotator) || !BaseScale3D.Equals(Rhs.BaseScale3D);
+	}
+
+	bool IsBaseObject(FString ObjectPath)
+	{
+		return (BaseObject == FSoftObjectPath(ObjectPath));
+	}
+};
+
+/**
+* Container struct for array of Options
+*/
+USTRUCT(BlueprintType)
+struct WAVEFUNCTIONCOLLAPSE_API FWaveFunctionCollapseOptions
+{
+	GENERATED_BODY()
+
+	UPROPERTY(BlueprintReadWrite, EditAnywhere, Category = "WaveFunctionCollapse")
+	TArray<FWaveFunctionCollapseOption> Options;
+};
+
+/**
+* Container struct for AdjacencyToOptionsMap
+* Stores the weight and contribution of an option
+*/
+USTRUCT(BlueprintType)
+struct WAVEFUNCTIONCOLLAPSE_API FWaveFunctionCollapseAdjacencyToOptionsMap
+{
+	GENERATED_BODY()
+
+	/**
+	* The amount of times an option is present when deriving a model.
+	* This value is used to calculate its weight.
+	*/
+	UPROPERTY(BlueprintReadWrite, EditAnywhere, Category = "WaveFunctionCollapse", meta = (ClampMin = "0"))
+	int32 Contribution = 1;
+
+	/**
+	* The weight of an option calculated by dividing this Contribution by the sum of all contributions of all options.
+	*/
+	UPROPERTY(BlueprintReadWrite, EditAnywhere, Category = "WaveFunctionCollapse", meta = (ClampMin = "0", ClampMax = "1"))
+	float Weight;
+
+	UPROPERTY(BlueprintReadWrite, EditAnywhere, Category = "WaveFunctionCollapse")
+	TMap<EWaveFunctionCollapseAdjacency, FWaveFunctionCollapseOptions> AdjacencyToOptionsMap;
+};
+
+/**
+* A Model of WFC constraints.
+* This data asset should contain all necessary data to allow for a WFC solve of an arbitrary grid size.
+*/
+UCLASS(Blueprintable, BlueprintType)
+class WAVEFUNCTIONCOLLAPSE_API UWaveFunctionCollapseModel : public UDataAsset
+{
+	GENERATED_BODY()
+
+public:
+	/**
+	* Grid Tile Size in cm^3
+	*/
+	UPROPERTY(BlueprintReadWrite, EditAnywhere, Category = "WaveFunctionCollapse")
+	float TileSize;
+
+	UPROPERTY(BlueprintReadWrite, EditAnywhere, Category = "WaveFunctionCollapse")
+	TMap<FWaveFunctionCollapseOption, FWaveFunctionCollapseAdjacencyToOptionsMap> Constraints;
+
+	UPROPERTY(BlueprintReadWrite, EditAnywhere, Category = "WaveFunctionCollapse", meta = (AllowedClasses = "StaticMesh, Blueprint"))
+	TArray<FSoftObjectPath> SpawnExclusion;
+
+	/**
+	* Create a constraint
+	* @param KeyOption Key option
+	* @param Adjacency Adjacency from KeyOption to AdjacentOption
+	* @param AdjacentOption Adjacent option
+	*/
+	UFUNCTION(BlueprintCallable, Category = "WaveFunctionCollapse")
+	void AddConstraint(const FWaveFunctionCollapseOption& KeyOption, EWaveFunctionCollapseAdjacency Adjacency, const FWaveFunctionCollapseOption& AdjacentOption);
+
+	/**
+	* Get all options for a given key option in a given adjacency
+	* @param KeyOption Key option
+	* @param Adjacency Adjacency from KeyOption to AdjacentOption
+	*/
+	UFUNCTION(BlueprintCallable, Category = "WaveFunctionCollapse")
+	FWaveFunctionCollapseOptions GetOptions(const FWaveFunctionCollapseOption& KeyOption, EWaveFunctionCollapseAdjacency Adjacency) const;
+
+	/**
+	* Set the weights of key objects based on their contribution values
+	*/
+	UFUNCTION(BlueprintCallable, Category = "WaveFunctionCollapse")
+	void SetWeightsFromContributions();
+
+	/**
+	* Set the weights of key objects to a given value
+	*/
+	UFUNCTION(BlueprintCallable, Category = "WaveFunctionCollapse")
+	void SetAllWeights(float Weight);
+
+	/**
+	* Set the contribution values of key objects to a given value
+	*/
+	UFUNCTION(BlueprintCallable, Category = "WaveFunctionCollapse")
+	void SetAllContributions(int32 Contribution);
+
+	/**
+	* Set the contribution value of a key object to a given value
+	*/
+	UFUNCTION(BlueprintCallable, Category = "WaveFunctionCollapse")
+	void SetOptionContribution(const FWaveFunctionCollapseOption& Option, int32 Contribution);
+
+	/**
+	* Get the weight value of an option
+	*/
+	UFUNCTION(BlueprintCallable, Category = "WaveFunctionCollapse")
+	float GetOptionWeight(const FWaveFunctionCollapseOption& Option) const;
+
+	/**
+	* Get the contribution value of an option
+	*/
+	UFUNCTION(BlueprintCallable, Category = "WaveFunctionCollapse")
+	int32 GetOptionContribution(const FWaveFunctionCollapseOption& Option) const;
+
+	/**
+	* Get the total count of constraints in this model
+	*/
+	UFUNCTION(BlueprintCallable, Category = "WaveFunctionCollapse")
+	int32 GetConstraintCount() const;
+
+	/**
+	* Swap meshes in the model with other meshes based on a map.  
+	* This is useful when working with template meshes that need to be swapped.
+	*/
+	UFUNCTION(BlueprintCallable, Category = "WaveFunctionCollapse")
+	void SwapMeshes(TMap<UStaticMesh*, UStaticMesh*> SourceToTargetMeshMap);
+};
+
+/**
+* Base Tile Struct which holds an array of remaining Options and its Shannon Entropy value
+*/
+USTRUCT(BlueprintType)
+struct WAVEFUNCTIONCOLLAPSE_API FWaveFunctionCollapseTile
+{
+	GENERATED_BODY()
+
+	UPROPERTY(BlueprintReadWrite, EditAnywhere, Category = "WaveFunctionCollapse")
+	float ShannonEntropy;
+	
+	UPROPERTY(BlueprintReadWrite, EditAnywhere, Category = "WaveFunctionCollapse")
+	TArray<FWaveFunctionCollapseOption> RemainingOptions;
+
+	FWaveFunctionCollapseTile() = default;
+
+	FWaveFunctionCollapseTile(const TArray<FWaveFunctionCollapseOption>& Options, float Entropy)
+	{
+		RemainingOptions = Options;
+		ShannonEntropy = Entropy;
+	}
+
+	// constructor with only one option
+	FWaveFunctionCollapseTile(const FWaveFunctionCollapseOption& Option, float Entropy)
+	{
+		RemainingOptions.Add(Option);
+		ShannonEntropy = Entropy;
+	}
+};
+
+/**
+* A helper struct used for queuing during Observation and Propagation phases
+*/
+USTRUCT(BlueprintType)
+struct WAVEFUNCTIONCOLLAPSE_API FWaveFunctionCollapseQueueElement
+{
+	GENERATED_BODY()
+
+	UPROPERTY(BlueprintReadWrite, EditAnywhere, Category = "WaveFunctionCollapse")
+	int32 CenterObjectIndex;
+
+	UPROPERTY(BlueprintReadWrite, EditAnywhere, Category = "WaveFunctionCollapse")
+	EWaveFunctionCollapseAdjacency Adjacency = EWaveFunctionCollapseAdjacency::Front;
+
+	FWaveFunctionCollapseQueueElement() = default;
+
+	FWaveFunctionCollapseQueueElement(int32 CenterObjectIndexInput, EWaveFunctionCollapseAdjacency AdjacencyInput)
+	{
+		CenterObjectIndex = CenterObjectIndexInput;
+		Adjacency = AdjacencyInput;
+	}
+};
\ No newline at end of file
diff --git a/Engine/Plugins/Experimental/WaveFunctionCollapse/Source/WaveFunctionCollapse/Public/WaveFunctionCollapseSubsystem.h b/Engine/Plugins/Experimental/WaveFunctionCollapse/Source/WaveFunctionCollapse/Public/WaveFunctionCollapseSubsystem.h
new file mode 100644
--- /dev/null
+++ b/Engine/Plugins/Experimental/WaveFunctionCollapse/Source/WaveFunctionCollapse/Public/WaveFunctionCollapseSubsystem.h
@@ -0,0 +1,175 @@
+// Copyright Epic Games, Inc. All Rights Reserved.
+
+#pragma once
+
+#include "EditorSubsystem.h"
+#include "WaveFunctionCollapseModel.h"
+#include "Components/ActorComponent.h"
+#include "WaveFunctionCollapseSubsystem.generated.h"
+
+WAVEFUNCTIONCOLLAPSE_API DECLARE_LOG_CATEGORY_EXTERN(LogWFC, Log, All);
+
+UCLASS()
+class WAVEFUNCTIONCOLLAPSE_API UWaveFunctionCollapseSubsystem : public UEditorSubsystem
+{
+	GENERATED_BODY()
+
+public:
+
+	UPROPERTY(EditAnywhere, BlueprintReadWrite, Category = "WFCSettings")
+	UWaveFunctionCollapseModel* WFCModel;
+
+	UPROPERTY(EditAnywhere, BlueprintReadWrite, Category = "WFCSettings")
+	FIntVector Resolution;
+
+	UPROPERTY(EditAnywhere, BlueprintReadWrite, Category = "WFCSettings")
+	FVector OriginLocation = FVector::ZeroVector;
+
+	UPROPERTY(EditAnywhere, BlueprintReadWrite, Category = "WFCSettings")
+	FRotator Orientation = FRotator::ZeroRotator;
+
+	UPROPERTY(EditAnywhere, BlueprintReadWrite, Category = "WFCSettings")
+	bool bUseEmptyBorder;
+
+	UPROPERTY(EditAnywhere, BlueprintReadWrite, Category = "WFCSettings")
+	TMap<FIntVector, FWaveFunctionCollapseOption> StarterOptions;
+
+	/**
+	* Solve a grid using a WFC model.  If successful, spawn an actor.
+	* @param TryCount Amount of times to attempt a successful solve
+	* @param RandomSeed Seed for deterministic results.  When this value is 0 the seed will be generated. Seed value will be logged during the solve.
+	*/
+	UFUNCTION(BlueprintCallable, Category = "WFCFunctions")
+	AActor* Collapse(int32 TryCount = 1, int32 RandomSeed = 0);
+
+	/**
+	* Initialize WFC process which sets up Tiles and RemainingTiles arrays
+	* Pre-populates Tiles with StarterOptions, BorderOptions and InitialTiles
+	* @param Tiles Array of tiles (by ref)
+	* @param RemainingTiles Array of remaining tile indices.  Semi-sorted: Min Entropy tiles at the front, the rest remains unsorted (by ref)
+	*/
+	UFUNCTION(BlueprintCallable, Category = "WFCFunctions")
+	void InitializeWFC(TArray<FWaveFunctionCollapseTile>& Tiles, TArray<int32>& RemainingTiles);
+	
+	/**
+	* Observation phase: 
+	* This process randomly selects one tile from minimum entropy tiles
+	* then randomly selects a valid option for that tile
+	* @param Tiles Array of tiles (by ref)
+	* @param RemainingTiles Array of remaining tile indices.  Semi-sorted: Min Entropy tiles at the front, the rest remains unsorted (by ref)
+	* @param ObservationQueue Array to store tiles that need to be checked whether remaining options are affected during propagation phase (by ref)
+	*/
+	UFUNCTION(BlueprintCallable, Category = "WFCFunctions")
+	bool Observe(TArray<FWaveFunctionCollapseTile>& Tiles, 
+		TArray<int32>& RemainingTiles, 
+		TMap<int32, FWaveFunctionCollapseQueueElement>& ObservationQueue,
+		int32 RandomSeed);
+	
+	/**
+	* Propagation phase: 
+	* This process checks if the selection made during the observation is valid by checking constraint validity with neighboring tiles. 
+	* Neighboring tiles may reduce their remaining options to include only valid options.
+	* If the remaining options of a tile were modified, the neighboring tiles of the modified tile will be added to a queue.
+	* During this process, if any contradiction (a tile with zero remaining options) is encountered, the current solve will fail.
+	* @param Tiles Array of tiles (by ref)
+	* @param RemainingTiles Array of remaining tile indices.  Semi-sorted: Min Entropy tiles at the front, the rest remains unsorted (by ref)
+	* @param ObservationQueue Array to store tiles that need to be checked whether remaining options are affected (by ref)
+	* @param PropagationCount Counter for propagation passes
+	*/
+	UFUNCTION(BlueprintCallable, Category = "WFCFunctions")
+	bool Propagate(TArray<FWaveFunctionCollapseTile>& Tiles, 
+		TArray<int32>& RemainingTiles, 
+		TMap<int32, FWaveFunctionCollapseQueueElement>& ObservationQueue, 
+		int32& PropagationCount);
+	
+	/**
+	* Recursive Observation and Propagation cycle
+	* @param Tiles Array of tiles (by ref)
+	* @param RemainingTiles Array of remaining tile indices (by ref)
+	* @param ObservationQueue Array to store tiles that need to be checked whether remaining options are affected (by ref)
+	*/
+	UFUNCTION(BlueprintCallable, Category = "WFCFunctions")
+	bool ObservationPropagation(TArray<FWaveFunctionCollapseTile>& Tiles, 
+		TArray<int32>& RemainingTiles,
+		TMap<int32, FWaveFunctionCollapseQueueElement>& ObservationQueue,
+		int32 RandomSeed);
+
+	/**
+	* Derive grid from the bounds of an array of transforms
+	* Assumptions:
+	*	-Transforms can only represent a single grid
+	*   -Sets empty starter option if there is a valid grid position with no transform
+	*   -Orientation is determined by the yaw of the first transform in the array
+	* @param Transforms Array of transforms (by ref)
+	*/
+	UFUNCTION(BlueprintCallable, Category = "WFCFunctions")
+	void DeriveGridFromTransformBounds(const TArray<FTransform>& Transforms);
+
+	/**
+	* Derive grid from an array of transforms
+	* Assumptions:
+	*   -Every transform represents the center point of a tile position
+	*   -Sets empty starter option if there is a valid grid position with no transform
+	*   -Orientation is determined by the yaw of the first transform in the array
+	* @param Transforms Array of transforms (by ref)
+	*/
+	UFUNCTION(BlueprintCallable, Category = "WFCFunctions")
+	void DeriveGridFromTransforms(const TArray<FTransform>& Transforms);
+
+private:
+
+	/**
+	* Builds the Initial Tile which is a tile containing all possible options
+	* @param InitialTile The Initial Tile (by ref)
+	*/
+	bool BuildInitialTile(FWaveFunctionCollapseTile& InitialTile);
+
+	/**
+	* Get valid options for a border tile
+	* @param Position Position of border tile
+	* @param TmpInitialOptions Initial options which should contain all possible options
+	*/
+	TArray<FWaveFunctionCollapseOption> GetInnerBorderOptions(FIntVector Position, const TArray<FWaveFunctionCollapseOption>& InitialOptions);
+
+	/**
+	* Used in GetBorderOptions to gather invalid options for border tiles that should be removed from the InitialOptions set
+	* @param Adjacency Direction from Exterior Border to Inner Border
+	* @param InInitialOptions This should be the InitialOptions
+	* @param OutBorderOptionsToRemove Array containing gathered options to remove
+	*/
+	void GatherInnerBorderOptionsToRemove(EWaveFunctionCollapseAdjacency Adjacency, const TArray<FWaveFunctionCollapseOption>& InitialOptions, TArray<FWaveFunctionCollapseOption>& OutBorderOptionsToRemove);
+	
+	/**
+	* Checks if a position is an inner border
+	* @param Position
+	*/
+	bool IsPositionInnerBorder(FIntVector Position);
+	
+	/**
+	* Used in Observe and Propagate to add adjacent indices to a queue
+	* @param CenterIndex Index of the center object
+	* @param RemainingTiles Used to check if index still remains in RemainingTiles
+	* @param OutQueue Queue to add indices to
+	*/
+	void AddAdjacentIndicesToQueue(int32 CenterIndex, const TArray<int32>& RemainingTiles, TMap<int32,FWaveFunctionCollapseQueueElement>& OutQueue);
+	
+	/**
+	* Add an Instance Component with a given name
+	* @param Actor Actor to add component to
+	* @param ComponentClass
+	* @param ComponentName
+	*/
+	UActorComponent* AddNamedInstanceComponent(AActor* Actor, TSubclassOf<UActorComponent> ComponentClass, FName ComponentName);
+	
+	/**
+	* Spawn an actor given an array of successfully solved tiles
+	* @param Tiles Successfully solved array of tiles
+	*/
+	AActor* SpawnActorFromTiles(const TArray<FWaveFunctionCollapseTile>& Tiles);
+	
+	/**
+	* Returns true if no remaining options in given tiles are an empty/void option or included in the SpawnExclusion list
+	* @param Tiles Successfully solved array of tiles
+	*/
+	bool AreAllTilesNonSpawnable(const TArray<FWaveFunctionCollapseTile>& Tiles);
+};
\ No newline at end of file
diff --git a/Engine/Plugins/Experimental/WaveFunctionCollapse/Source/WaveFunctionCollapse/WaveFunctionCollapse.Build.cs b/Engine/Plugins/Experimental/WaveFunctionCollapse/Source/WaveFunctionCollapse/WaveFunctionCollapse.Build.cs
new file mode 100644
--- /dev/null
+++ b/Engine/Plugins/Experimental/WaveFunctionCollapse/Source/WaveFunctionCollapse/WaveFunctionCollapse.Build.cs
@@ -0,0 +1,41 @@
+// Copyright Epic Games, Inc. All Rights Reserved.
+
+using UnrealBuildTool;
+
+public class WaveFunctionCollapse : ModuleRules
+{
+	public WaveFunctionCollapse(ReadOnlyTargetRules Target) : base(Target)
+	{
+		ShortName = "WFC";
+
+		PCHUsage = ModuleRules.PCHUsageMode.UseExplicitOrSharedPCHs;
+		
+		PublicDependencyModuleNames.AddRange(
+			new string[]
+			{
+				"Core",
+				"EditorSubsystem",
+				"EditorStyle",
+				"LevelEditor",
+				"EditorFramework",
+				"UnrealEd",
+				"Blutility",
+				"UMG",
+				"UMGEditor",
+				"AssetTools",
+				"PropertyEditor",
+			}
+		);
+		
+		PrivateDependencyModuleNames.AddRange(
+			new string[]
+			{
+				"CoreUObject",
+				"InputCore",
+				"Engine",
+				"Slate",
+				"SlateCore",
+			}
+		);
+	}
+}
diff --git a/Engine/Plugins/Experimental/WaveFunctionCollapse/WaveFunctionCollapse.uplugin b/Engine/Plugins/Experimental/WaveFunctionCollapse/WaveFunctionCollapse.uplugin
new file mode 100644
--- /dev/null
+++ b/Engine/Plugins/Experimental/WaveFunctionCollapse/WaveFunctionCollapse.uplugin
@@ -0,0 +1,23 @@
+{
+	"FileVersion": 1,
+	"Version": 1,
+	"VersionName": "1.0",
+	"FriendlyName": "Wave Function Collapse (Experimental)",
+	"Description": "Wave Function Collapse tools for tile-based model synthesis",
+	"Category": "Experimental",
+	"CreatedBy": "Epic Games, Inc.",
+	"CreatedByURL": "https://www.epicgames.com/",
+	"DocsURL": "",
+	"MarketplaceURL": "",
+	"SupportURL": "",
+	"CanContainContent": true,
+	"IsExperimentalVersion": true,
+	"Installed": true,
+	"Modules": [
+		{
+			"Name": "WaveFunctionCollapse",
+			"Type": "Editor",
+			"LoadingPhase": "PostEngineInit"
+		}
+	]
+}
\ No newline at end of file
