WFC uplugin "Installed" set to false

diff --git a/Engine/Plugins/Experimental/WaveFunctionCollapse/WaveFunctionCollapse.uplugin b/Engine/Plugins/Experimental/WaveFunctionCollapse/WaveFunctionCollapse.uplugin
--- a/Engine/Plugins/Experimental/WaveFunctionCollapse/WaveFunctionCollapse.uplugin
+++ b/Engine/Plugins/Experimental/WaveFunctionCollapse/WaveFunctionCollapse.uplugin
@@ -12,7 +12,7 @@
 	"SupportURL": "",
 	"CanContainContent": true,
 	"IsExperimentalVersion": true,
-	"Installed": true,
+	"Installed": false,
 	"Modules": [
 		{
 			"Name": "WaveFunctionCollapse",
