[WFC] Use generated class for wave function collapse with blueprints

#11588

diff --git a/Engine/Plugins/Experimental/WaveFunctionCollapse/Source/WaveFunctionCollapse/Private/WaveFunctionCollapseSubsystem.cpp b/Engine/Plugins/Experimental/WaveFunctionCollapse/Source/WaveFunctionCollapse/Private/WaveFunctionCollapseSubsystem.cpp
--- a/Engine/Plugins/Experimental/WaveFunctionCollapse/Source/WaveFunctionCollapse/Private/WaveFunctionCollapseSubsystem.cpp
+++ b/Engine/Plugins/Experimental/WaveFunctionCollapse/Source/WaveFunctionCollapse/Private/WaveFunctionCollapseSubsystem.cpp
@@ -633,12 +633,14 @@ AActor* UWaveFunctionCollapseSubsystem::SpawnActorFromTiles(const TArray<FWaveFu
 			// Blueprints are handled with ChildActorComponents
 			else if (UBlueprint* LoadedBlueprint = Cast<UBlueprint>(LoadedObject))
 			{
+				const TSubclassOf<UObject> GeneratedClass = LoadedBlueprint->GeneratedClass;
+
 				// if BP is placeable
-				if (!(LoadedBlueprint->GetClass()->HasAnyClassFlags(CLASS_NotPlaceable | CLASS_Deprecated | CLASS_Abstract))
-					&& LoadedBlueprint->GetClass()->IsChildOf(AActor::StaticClass()))
+				if (!(GeneratedClass->GetClass()->HasAnyClassFlags(CLASS_NotPlaceable | CLASS_Deprecated | CLASS_Abstract))
+					&& GeneratedClass->GetClass()->IsChildOf(AActor::StaticClass()))
 				{
 					UChildActorComponent* ChildActorComponent = Cast<UChildActorComponent>(UWaveFunctionCollapseSubsystem::AddNamedInstanceComponent(SpawnedActor, UChildActorComponent::StaticClass(), LoadedObject->GetFName()));
-					ChildActorComponent->SetChildActorClass(LoadedBlueprint->GetClass());
+					ChildActorComponent->SetChildActorClass(*GeneratedClass);
 					ChildActorComponent->SetRelativeLocation(TilePosition);
 					ChildActorComponent->SetRelativeRotation(BaseRotator);
 					ChildActorComponent->SetRelativeScale3D(BaseScale3D);
