Replace short or incorrect path names in metadata with long or correct ones

diff --git a/Engine/Plugins/Experimental/WaveFunctionCollapse/Source/WaveFunctionCollapse/Public/WaveFunctionCollapseBPLibrary.h b/Engine/Plugins/Experimental/WaveFunctionCollapse/Source/WaveFunctionCollapse/Public/WaveFunctionCollapseBPLibrary.h
--- a/Engine/Plugins/Experimental/WaveFunctionCollapse/Source/WaveFunctionCollapse/Public/WaveFunctionCollapseBPLibrary.h
+++ b/Engine/Plugins/Experimental/WaveFunctionCollapse/Source/WaveFunctionCollapse/Public/WaveFunctionCollapseBPLibrary.h
@@ -17,7 +17,7 @@ struct WAVEFUNCTIONCOLLAPSE_API FWaveFunctionCollapseNeighbor
 	UPROPERTY(BlueprintReadWrite, EditAnywhere, Category = "WaveFunctionCollapse")
 	int32 Step = 1;
 
-	UPROPERTY(BlueprintReadWrite, EditAnywhere, Category = "WaveFunctionCollapse", meta = (AllowedClasses = "StaticMesh, Blueprint"))
+	UPROPERTY(BlueprintReadWrite, EditAnywhere, Category = "WaveFunctionCollapse", meta = (AllowedClasses = "/Script/Engine.StaticMesh, /Script/Engine.Blueprint"))
 	FSoftObjectPath NeighborObject;
 };
 
@@ -26,13 +26,13 @@ struct WAVEFUNCTIONCOLLAPSE_API FWaveFunctionCollapseNeighborRule
 {
 	GENERATED_BODY()
 
-	UPROPERTY(BlueprintReadWrite, EditAnywhere, Category = "WaveFunctionCollapse", meta = (AllowedClasses = "StaticMesh, Blueprint"))
+	UPROPERTY(BlueprintReadWrite, EditAnywhere, Category = "WaveFunctionCollapse", meta = (AllowedClasses = "/Script/Engine.StaticMesh, /Script/Engine.Blueprint"))
 	FSoftObjectPath KeyObject;
 
 	UPROPERTY(BlueprintReadWrite, EditAnywhere, Category = "WaveFunctionCollapse")
 	TArray<FWaveFunctionCollapseNeighbor> Neighbors;
 
-	UPROPERTY(BlueprintReadWrite, EditAnywhere, Category = "WaveFunctionCollapse", meta = (AllowedClasses = "StaticMesh, Blueprint"))
+	UPROPERTY(BlueprintReadWrite, EditAnywhere, Category = "WaveFunctionCollapse", meta = (AllowedClasses = "/Script/Engine.StaticMesh, /Script/Engine.Blueprint"))
 	TArray<FSoftObjectPath> SpawnObjects;
 
 	UPROPERTY(BlueprintReadWrite, EditAnywhere, Category = "WaveFunctionCollapse")
diff --git a/Engine/Plugins/Experimental/WaveFunctionCollapse/Source/WaveFunctionCollapse/Public/WaveFunctionCollapseModel.h b/Engine/Plugins/Experimental/WaveFunctionCollapse/Source/WaveFunctionCollapse/Public/WaveFunctionCollapseModel.h
--- a/Engine/Plugins/Experimental/WaveFunctionCollapse/Source/WaveFunctionCollapse/Public/WaveFunctionCollapseModel.h
+++ b/Engine/Plugins/Experimental/WaveFunctionCollapse/Source/WaveFunctionCollapse/Public/WaveFunctionCollapseModel.h
@@ -25,7 +25,7 @@ struct WAVEFUNCTIONCOLLAPSE_API FWaveFunctionCollapseOption
 {
 	GENERATED_BODY()
 
-	UPROPERTY(BlueprintReadWrite, EditAnywhere, Category = "WaveFunctionCollapse", meta = (AllowedClasses = "StaticMesh, Blueprint"))
+	UPROPERTY(BlueprintReadWrite, EditAnywhere, Category = "WaveFunctionCollapse", meta = (AllowedClasses = "/Script/Engine.StaticMesh, /Script/Engine.Blueprint"))
 	FSoftObjectPath BaseObject;
 
 	UPROPERTY(BlueprintReadWrite, EditAnywhere, Category = "WaveFunctionCollapse")
@@ -133,7 +133,7 @@ public:
 	UPROPERTY(BlueprintReadWrite, EditAnywhere, Category = "WaveFunctionCollapse")
 	TMap<FWaveFunctionCollapseOption, FWaveFunctionCollapseAdjacencyToOptionsMap> Constraints;
 
-	UPROPERTY(BlueprintReadWrite, EditAnywhere, Category = "WaveFunctionCollapse", meta = (AllowedClasses = "StaticMesh, Blueprint"))
+	UPROPERTY(BlueprintReadWrite, EditAnywhere, Category = "WaveFunctionCollapse", meta = (AllowedClasses = "/Script/Engine.StaticMesh, /Script/Engine.Blueprint"))
 	TArray<FSoftObjectPath> SpawnExclusion;
 
 	/**
diff --git a/Engine/Plugins/VirtualProduction/Avalanche/Source/AvalancheViewport/Public/AvaViewportSettings.h b/Engine/Plugins/VirtualProduction/Avalanche/Source/AvalancheViewport/Public/AvaViewportSettings.h
--- a/Engine/Plugins/VirtualProduction/Avalanche/Source/AvalancheViewport/Public/AvaViewportSettings.h
+++ b/Engine/Plugins/VirtualProduction/Avalanche/Source/AvalancheViewport/Public/AvaViewportSettings.h
@@ -109,7 +109,7 @@ public:
 	UPROPERTY(Config, EditAnywhere, BlueprintReadWrite, Category = "Pixel Grid")
 	FLinearColor PixelGridColor;
 
-	UPROPERTY(Config, EditAnywhere, BlueprintReadWrite, Category = "Snapping", Meta = (Bitmask, BitmaskEnum = "/Script/AvalancheViewport/EAvaViewportSnapState"))
+	UPROPERTY(Config, EditAnywhere, BlueprintReadWrite, Category = "Snapping", Meta = (Bitmask, BitmaskEnum = "/Script/AvalancheViewport.EAvaViewportSnapState"))
 	int32 SnapState;
 
 	UPROPERTY(Config, EditAnywhere, BlueprintReadWrite, Category = "Snapping")
